Cherrypicked from 21659075:

Adding a small sleep to the capture complete polling in leiu of an actual fence wait mechanism.
This is an attempt to reduce editor idle cpu usage.

diff --git a/Engine/Plugins/Media/PixelCapture/Source/PixelCapture/Private/PixelCaptureCapturerRHI.cpp b/Engine/Plugins/Media/PixelCapture/Source/PixelCapture/Private/PixelCaptureCapturerRHI.cpp
--- a/Engine/Plugins/Media/PixelCapture/Source/PixelCapture/Private/PixelCaptureCapturerRHI.cpp
+++ b/Engine/Plugins/Media/PixelCapture/Source/PixelCapture/Private/PixelCaptureCapturerRHI.cpp
@@ -70,6 +70,8 @@ void FPixelCaptureCapturerRHI::CheckComplete()
 	{
 		TSharedRef<FPixelCaptureCapturerRHI> ThisRHIRef = StaticCastSharedRef<FPixelCaptureCapturerRHI>(AsShared());
 		AsyncTask(ENamedThreads::AnyHiPriThreadHiPriTask, [ThisRHIRef]() {
+			// we want to check quickly but we dont want to spam new tasks. at some point you're just making more work.
+			FPlatformProcess::Sleep(0.0001f);
 			ThisRHIRef->CheckComplete();
 		});
 	}
