Cherrypicked from 21695388:

Better handling of fence polling so we dont recursively spam new tasks.

diff --git a/Engine/Plugins/Media/PixelCapture/Source/PixelCapture/Private/PixelCaptureCapturerRHI.cpp b/Engine/Plugins/Media/PixelCapture/Source/PixelCapture/Private/PixelCaptureCapturerRHI.cpp
--- a/Engine/Plugins/Media/PixelCapture/Source/PixelCapture/Private/PixelCaptureCapturerRHI.cpp
+++ b/Engine/Plugins/Media/PixelCapture/Source/PixelCapture/Private/PixelCaptureCapturerRHI.cpp
@@ -66,19 +66,17 @@ void FPixelCaptureCapturerRHI::BeginProcess(const IPixelCaptureInputFrame& Input
 
 void FPixelCaptureCapturerRHI::CheckComplete()
 {
-	if (!Fence->Poll())
-	{
-		TSharedRef<FPixelCaptureCapturerRHI> ThisRHIRef = StaticCastSharedRef<FPixelCaptureCapturerRHI>(AsShared());
-		AsyncTask(ENamedThreads::AnyHiPriThreadHiPriTask, [ThisRHIRef]() {
-			// we want to check quickly but we dont want to spam new tasks. at some point you're just making more work.
+	// in lieu of a proper callback we need to capture a thread to poll the fence
+	// so we know as quickly as possible when we can readback.
+	TSharedRef<FPixelCaptureCapturerRHI> ThisRHIRef = StaticCastSharedRef<FPixelCaptureCapturerRHI>(AsShared());
+	AsyncTask(ENamedThreads::AnyHiPriThreadHiPriTask, [ThisRHIRef]() {
+		while (!ThisRHIRef->Fence->Poll())
+		{
+			// we want to check quickly but we dont want to burn CPU.
 			FPlatformProcess::Sleep(0.0001f);
-			ThisRHIRef->CheckComplete();
-		});
-	}
-	else
-	{
-		OnRHIStageComplete();
-	}
+		}
+		ThisRHIRef->OnRHIStageComplete();
+	});
 }
 
 void FPixelCaptureCapturerRHI::OnRHIStageComplete()
