Cherrypicked from 21693502:

Fixing some multi format capture complete issues that were showing up in some broken test.
- We now handle the case of capturers completing immediately requests to new formats happening while capture is still starting.

diff --git a/Engine/Plugins/Media/PixelCapture/Source/PixelCapture/Private/PixelCaptureCapturerMultiFormat.cpp b/Engine/Plugins/Media/PixelCapture/Source/PixelCapture/Private/PixelCaptureCapturerMultiFormat.cpp
--- a/Engine/Plugins/Media/PixelCapture/Source/PixelCapture/Private/PixelCaptureCapturerMultiFormat.cpp
+++ b/Engine/Plugins/Media/PixelCapture/Source/PixelCapture/Private/PixelCaptureCapturerMultiFormat.cpp
@@ -41,8 +41,6 @@ void FPixelCaptureCapturerMultiFormat::Capture(const IPixelCaptureInputFrame& So
 		return;
 	}
 
-	FScopeLock FormatLock(&FormatGuard);
-
 	// cache each layer size
 	if (LayerSizes.IsEmpty())
 	{
@@ -56,18 +54,32 @@ void FPixelCaptureCapturerMultiFormat::Capture(const IPixelCaptureInputFrame& So
 		}
 	}
 
-	// start all the format captures
-	PendingFormats = FormatCapturers.Num();
-	for (auto& FormatCapturer : FormatCapturers)
+	if (FormatCapturers.Num() > 0)
 	{
-		FormatCapturer.Value->Capture(SourceFrame);
+		TArray<TSharedPtr<FPixelCaptureCapturerLayered>> TempCopy;
+		{
+			// iterate a temp copy so modifications to the map in Capture
+			// (and their callback events) is ok.
+			FScopeLock FormatLock(&FormatGuard);
+			for (auto& FormatCapturer : FormatCapturers)
+			{
+				TempCopy.Add(FormatCapturer.Value);
+			}
+		}
+		// start all the format captures
+		PendingFormats = FormatCapturers.Num();
+		for (auto& Capturer : TempCopy)
+		{
+			Capturer->Capture(SourceFrame);
+		}
 	}
-
-	// if we have no pipelines setup we will never fire off the complete event
-	// fake it here by "pretending" like the capture is complete
-	if (PendingFormats == 0)
+	else
 	{
-		OnComplete.Broadcast();
+		// if we have no capturers just signal this capture as complete
+		if (PendingFormats == 0)
+		{
+			OnComplete.Broadcast();
+		}
 	}
 }
 
@@ -127,7 +139,7 @@ void FPixelCaptureCapturerMultiFormat::OnDisconnected()
 void FPixelCaptureCapturerMultiFormat::OnCaptureFormatComplete(int32 Format)
 {
 	--PendingFormats;
-	CheckFormatEvent(Format);
+	CheckFormatEvent(Format); // checks if anything is waiting on this format
 	if (PendingFormats == 0)
 	{
 		OnComplete.Broadcast();
