Cherrypicked from 21735551 with additional NonUnity fixes: PixelCapture tests.

diff --git a/Engine/Plugins/Media/PixelCapture/Source/PixelCapture/Private/Tests/PixelCaptureTests.cpp b/Engine/Plugins/Media/PixelCapture/Source/PixelCapture/Private/Tests/PixelCaptureTests.cpp
new file mode 100644
--- /dev/null
+++ b/Engine/Plugins/Media/PixelCapture/Source/PixelCapture/Private/Tests/PixelCaptureTests.cpp
@@ -0,0 +1,259 @@
+// Copyright Epic Games, Inc. All Rights Reserved.
+
+#include "Misc/AutomationTest.h"
+#include "Tests/AutomationCommon.h"
+#include "PixelCaptureBufferFormat.h"
+#include "PixelCaptureCapturer.h"
+#include "PixelCaptureCapturerLayered.h"
+#include "PixelCaptureCapturerMultiFormat.h"
+
+#if WITH_DEV_AUTOMATION_TESTS
+
+namespace
+{
+	enum class EMockFormats : int32
+	{
+		Input = PixelCaptureBufferFormat::FORMAT_USER,
+		Output1,
+		Output2,
+		Output3,
+	};
+
+	struct FMockInput : public IPixelCaptureInputFrame
+	{
+		int32 Width;
+		int32 Height;
+		int32 MagicData;
+
+		FMockInput(int32 InWidth, int32 InHeight, int32 InMagicData)
+			: Width(InWidth)
+			, Height(InHeight)
+			, MagicData(InMagicData)
+		{
+		}
+
+		virtual int32 GetType() const override { return StaticCast<int32>(EMockFormats::Input); }
+		virtual int32 GetWidth() const override { return Width; }
+		virtual int32 GetHeight() const override { return Height; }
+	};
+
+	struct FMockOutput : public IPixelCaptureOutputFrame
+	{
+		int32 Width = 0;
+		int32 Height = 0;
+		int32 MagicData = 0;
+
+		virtual int32 GetWidth() const override { return Width; }
+		virtual int32 GetHeight() const override { return Height; }
+	};
+
+	struct FMockCapturer : public FPixelCaptureCapturer
+	{
+		int32 Format;
+		float Scale;
+		bool AutoComplete;
+
+		FMockCapturer(int32 InFormat, float InScale, bool InAutoComplete)
+			: Format(InFormat)
+			, Scale(InScale)
+			, AutoComplete(InAutoComplete)
+		{
+		}
+
+		virtual FString GetCapturerName() const override { return "Mock Capturer"; }
+
+		virtual IPixelCaptureOutputFrame* CreateOutputBuffer(int32 InputWidth, int32 InputHeight) override
+		{
+			return new FMockOutput();
+		}
+
+		virtual void BeginProcess(const IPixelCaptureInputFrame& InputFrame, IPixelCaptureOutputFrame* OutputBuffer) override
+		{
+			check(InputFrame.GetType() == StaticCast<int32>(EMockFormats::Input));
+			const FMockInput& Input = StaticCast<const FMockInput&>(InputFrame);
+			FMockOutput* Output = StaticCast<FMockOutput*>(OutputBuffer);
+			Output->Width = Input.Width * Scale;
+			Output->Height = Input.Height * Scale;
+			Output->MagicData = Input.MagicData;
+
+			if (AutoComplete)
+			{
+				EndProcess();
+			}
+		}
+
+		void MockProcessComplete()
+		{
+			EndProcess();
+		}
+	};
+
+	struct FMockCaptureSource : public IPixelCaptureCapturerSource
+	{
+		TAtomic<int> FrameCompleteCount = 0;
+		TMap<FMockCapturer*, int> LayerCompleteCounts;
+		FCriticalSection CountSection;
+
+		virtual ~FMockCaptureSource() = default;
+
+		virtual TSharedPtr<FPixelCaptureCapturer> CreateCapturer(int32 FinalFormat, float FinalScale) override
+		{
+			checkf(FinalFormat == StaticCast<int32>(EMockFormats::Output1)
+					|| FinalFormat == StaticCast<int32>(EMockFormats::Output2)
+					|| FinalFormat == StaticCast<int32>(EMockFormats::Output3),
+				TEXT("Unknown destination format."));
+			TSharedPtr<FMockCapturer> NewCapturer = MakeShared<FMockCapturer>(FinalFormat, FinalScale, true);
+			{
+				FScopeLock Lock(&CountSection);
+				LayerCompleteCounts.Add(NewCapturer.Get(), 0);
+			}
+			NewCapturer->OnComplete.AddRaw(this, &FMockCaptureSource::OnLayerComplete, NewCapturer.Get());
+			return NewCapturer;
+		}
+
+		void OnLayerComplete(FMockCapturer* Capturer)
+		{
+			FScopeLock Lock(&CountSection);
+			LayerCompleteCounts[Capturer]++;
+		}
+	};
+} // namespace
+
+namespace UE::PixelCapture
+{
+	IMPLEMENT_SIMPLE_AUTOMATION_TEST(MultiCaptureTest, "System.Plugins.PixelCapture.MultiCaptureTest", EAutomationTestFlags::EditorContext | EAutomationTestFlags::ClientContext | EAutomationTestFlags::ProductFilter)
+	bool MultiCaptureTest::RunTest(const FString& Parameters)
+	{
+		const TArray<float> Scales = { 1.0f, 0.5f, 0.25f };
+
+		// NOTE: resolution changes with the same capturer are not supported and the capturer will assert.
+		// ideally we would make sure that assert gets fired but AddExpectedError only works on error logs
+		const int32 InputWidth = 1024;
+		const int32 InputHeight = 768;
+		const int32 InputMagicData = 0xDEADBEEF;
+
+		// basic capture behaviour
+		{
+			FMockCapturer Capturer(StaticCast<int32>(EMockFormats::Output1), 1.0f, false);
+			TestTrue("Capturer starts uninitialized", Capturer.IsInitialized() == false);
+			TestTrue("Capturer starts not busy", Capturer.IsBusy() == false);
+			TestTrue("Capturer starts without output", Capturer.HasOutput() == false);
+			TestTrue("Capturer returns null if no input present", Capturer.ReadOutput() == nullptr);
+			Capturer.Capture(FMockInput(InputWidth, InputHeight, InputMagicData));
+			TestTrue("Capturer initialized after frame input", Capturer.IsInitialized() == true);
+			TestTrue("Capturer busy after frame input", Capturer.IsBusy() == true);
+			TestTrue("Capturer has no output before completing", Capturer.HasOutput() == false);
+			TestTrue("Capturer still returns null after input but not complete", Capturer.ReadOutput() == nullptr);
+			Capturer.MockProcessComplete();
+			TestTrue("Capturer still initialized after process complete", Capturer.IsInitialized() == true);
+			TestTrue("Capturer no longer busy after process complete", Capturer.IsBusy() == false);
+			TestTrue("Capturer has output after completing", Capturer.HasOutput() == true);
+			TestTrue("Capturer returns output after process complete", Capturer.ReadOutput() != nullptr);
+		}
+
+		FMockCaptureSource CaptureSource;
+
+		// testing simple layered capturer
+		{
+			TSharedPtr<FPixelCaptureCapturerLayered> LayeredCapturer = FPixelCaptureCapturerLayered::Create(&CaptureSource, StaticCast<int32>(EMockFormats::Output1), Scales);
+			LayeredCapturer->OnComplete.AddLambda([&]() { CaptureSource.FrameCompleteCount++; });
+			LayeredCapturer->Capture(FMockInput(InputWidth, InputHeight, InputMagicData));
+			TestTrue("Frame complete callback called once.", CaptureSource.FrameCompleteCount == 1);
+			for (TMap<FMockCapturer*, int>::ElementType& CaptureComplete : CaptureSource.LayerCompleteCounts)
+			{
+				TestTrue(FString::Printf(TEXT("Layer %p called complete callback once."), CaptureComplete.Key), CaptureComplete.Value == 1);
+			}
+			for (int i = 0; i < Scales.Num(); ++i)
+			{
+				TSharedPtr<FMockOutput> Layer = StaticCastSharedPtr<FMockOutput>(LayeredCapturer->ReadOutput(i));
+				TestTrue(FString::Printf(TEXT("Layer %d has output with correct values."), i),
+					Layer->Width == InputWidth * Scales[i] && Layer->Height == InputHeight * Scales[i] && Layer->MagicData == InputMagicData);
+			}
+		}
+
+		TSharedPtr<FPixelCaptureCapturerMultiFormat> FormatCapturer = FPixelCaptureCapturerMultiFormat::Create(&CaptureSource, Scales);
+
+		// testing mutli format capture
+		{
+			const int32 TestFormat = StaticCast<int>(EMockFormats::Output1);
+
+			// test that WaitForFormat will time out
+			{
+				const double PreWaitSeconds = FPlatformTime::Seconds();
+				FormatCapturer->WaitForFormat(TestFormat, 0, 500);
+				const double DeltaSeconds = FPlatformTime::Seconds() - PreWaitSeconds;
+				TestTrue("WaitForFormat times out in a reasonable time.", FMath::Abs(DeltaSeconds) < 600);
+			}
+
+			// Test output after implicit format add from WaitForFormat
+			{
+				FormatCapturer->Capture(FMockInput(InputWidth, InputHeight, InputMagicData));
+				TestTrue("MultiFormat Capturer reports correct layer count after input.", FormatCapturer->GetNumLayers() == Scales.Num());
+				for (int i = 0; i < Scales.Num(); ++i)
+				{
+					TSharedPtr<FMockOutput> Layer = StaticCastSharedPtr<FMockOutput>(FormatCapturer->RequestFormat(TestFormat, i));
+					TestTrue(FString::Printf(TEXT("MultiFormat Capturer reports correct layer %d width after input."), i), FormatCapturer->GetWidth(i) == InputWidth * Scales[i]);
+					TestTrue(FString::Printf(TEXT("MultiFormat Capturer reports correct layer %d height after input."), i), FormatCapturer->GetHeight(i) == InputHeight * Scales[i]);
+					TestTrue(FString::Printf(TEXT("Format %d Layer %d has output with correct values."), TestFormat, i),
+						Layer->Width == InputWidth * Scales[i] && Layer->Height == InputHeight * Scales[i] && Layer->MagicData == InputMagicData);
+				}
+			}
+
+			TestTrue("WaitForFormat returns not null after frame input.", FormatCapturer->WaitForFormat(TestFormat, 0, 10) != nullptr);
+		}
+
+		// test explicit format add
+		{
+			const int32 TestFormat = StaticCast<int>(EMockFormats::Output2);
+			FormatCapturer->AddOutputFormat(TestFormat);
+			FormatCapturer->Capture(FMockInput(InputWidth, InputHeight, InputMagicData));
+			for (int i = 0; i < Scales.Num(); ++i)
+			{
+				TSharedPtr<FMockOutput> Layer = StaticCastSharedPtr<FMockOutput>(FormatCapturer->RequestFormat(TestFormat, i));
+				TestTrue(FString::Printf(TEXT("Format %d Layer %d has output with correct values."), TestFormat, i),
+					Layer->Width == InputWidth * Scales[i] && Layer->Height == InputHeight * Scales[i] && Layer->MagicData == InputMagicData);
+			}
+		}
+
+		// make sure all formats get new frames
+		{
+			const int32 InputMagicData2 = 0xBAADF00D;
+			FormatCapturer->Capture(FMockInput(InputWidth, InputHeight, InputMagicData2));
+
+			auto CheckFormats = { StaticCast<int>(EMockFormats::Output1), StaticCast<int>(EMockFormats::Output2) };
+			for (auto& TestFormat : CheckFormats)
+			{
+				for (int i = 0; i < Scales.Num(); ++i)
+				{
+					TSharedPtr<FMockOutput> Layer = StaticCastSharedPtr<FMockOutput>(FormatCapturer->RequestFormat(TestFormat, i));
+					TestTrue(FString::Printf(TEXT("Format %d Layer %d has output with correct values."), TestFormat, i),
+						Layer->Width == InputWidth * Scales[i] && Layer->Height == InputHeight * Scales[i] && Layer->MagicData == InputMagicData2);
+				}
+			}
+		}
+
+		// test the callbacks
+		{
+			CaptureSource.FrameCompleteCount = 0;
+			CaptureSource.LayerCompleteCounts.Empty();
+
+			// new capturer to force clear resolutions
+			FormatCapturer = FPixelCaptureCapturerMultiFormat::Create(&CaptureSource, Scales);
+			FormatCapturer->OnComplete.AddLambda([&]() { CaptureSource.FrameCompleteCount++; });
+
+			FormatCapturer->AddOutputFormat(StaticCast<int>(EMockFormats::Output1));
+			FormatCapturer->AddOutputFormat(StaticCast<int>(EMockFormats::Output2));
+			FormatCapturer->Capture(FMockInput(InputWidth, InputHeight, InputMagicData));
+
+			TestTrue("MultiFormat Frame complete callback called once.", CaptureSource.FrameCompleteCount == 1);
+			for (TMap<FMockCapturer*, int>::ElementType& CaptureComplete : CaptureSource.LayerCompleteCounts)
+			{
+				TestTrue(FString::Printf(TEXT("Format %d Layer %p called complete callback once."), CaptureComplete.Key->Format, CaptureComplete.Key), CaptureComplete.Value == 1);
+			}
+		}
+
+		return true;
+	}
+} // namespace UE::PixelCapture
+
+#endif // WITH_DEV_AUTOMATION_TESTS
