Cherrypicked from 21752202: Proper freeing of waiting events on the multi format capture.

diff --git a/Engine/Plugins/Media/PixelCapture/Source/PixelCapture/Private/PixelCaptureCapturerMultiFormat.cpp b/Engine/Plugins/Media/PixelCapture/Source/PixelCapture/Private/PixelCaptureCapturerMultiFormat.cpp
--- a/Engine/Plugins/Media/PixelCapture/Source/PixelCapture/Private/PixelCaptureCapturerMultiFormat.cpp
+++ b/Engine/Plugins/Media/PixelCapture/Source/PixelCapture/Private/PixelCaptureCapturerMultiFormat.cpp
@@ -124,7 +124,7 @@ TSharedPtr<IPixelCaptureOutputFrame> FPixelCaptureCapturerMultiFormat::WaitForFo
 	if (FEvent* Event = GetEventForFormat(Format))
 	{
 		Event->Wait(MaxWaitTime);
-		FreeEvent(Event);
+		FreeEvent(Format, Event);
 		return RequestFormat(Format, LayerIndex);
 	}
 	return nullptr;
@@ -169,8 +169,10 @@ void FPixelCaptureCapturerMultiFormat::CheckFormatEvent(int32 Format)
 	}
 }
 
-void FPixelCaptureCapturerMultiFormat::FreeEvent(FEvent* Event)
+void FPixelCaptureCapturerMultiFormat::FreeEvent(int32 Format, FEvent* Event)
 {
+	FScopeLock Lock(&EventMutex);
+	FormatEvents.Remove(Format);
 	FPlatformProcess::ReturnSynchEventToPool(Event);
 }
 
diff --git a/Engine/Plugins/Media/PixelCapture/Source/PixelCapture/Public/PixelCaptureCapturerMultiFormat.h b/Engine/Plugins/Media/PixelCapture/Source/PixelCapture/Public/PixelCaptureCapturerMultiFormat.h
--- a/Engine/Plugins/Media/PixelCapture/Source/PixelCapture/Public/PixelCaptureCapturerMultiFormat.h
+++ b/Engine/Plugins/Media/PixelCapture/Source/PixelCapture/Public/PixelCaptureCapturerMultiFormat.h
@@ -93,7 +93,7 @@ protected:
 	
 	FEvent* GetEventForFormat(int32 Format);
 	void CheckFormatEvent(int32 Format);
-	void FreeEvent(FEvent* Event);
+	void FreeEvent(int32 Format, FEvent* Event);
 	void FlushWaitingEvents();
 
 	IPixelCaptureCapturerSource* CapturerSource;
