Fix a regression in stream framerate caused by using a fence when copying textures.

The new code path utilises RDG for copying textures and is ~100% more performant. The old pathway which uses the RHI Command List and fences can be enabled by passing `-PixelCaptureUseFence` when launching an application.

diff --git a/Engine/Plugins/Media/PixelCapture/Source/PixelCapture/Private/PixelCaptureCapturerRHIRDG.cpp b/Engine/Plugins/Media/PixelCapture/Source/PixelCapture/Private/PixelCaptureCapturerRHIRDG.cpp
new file mode 100644
--- /dev/null
+++ b/Engine/Plugins/Media/PixelCapture/Source/PixelCapture/Private/PixelCaptureCapturerRHIRDG.cpp
@@ -0,0 +1,70 @@
+// Copyright Epic Games, Inc. All Rights Reserved.
+
+#include "PixelCaptureCapturerRHIRDG.h"
+#include "PixelCaptureInputFrameRHI.h"
+#include "PixelCaptureOutputFrameRHI.h"
+#include "PixelCaptureBufferFormat.h"
+#include "PixelCaptureUtils.h"
+#include "RenderGraphUtils.h"
+#include "Async/Async.h"
+
+TSharedPtr<FPixelCaptureCapturerRHIRDG> FPixelCaptureCapturerRHIRDG::Create(float InScale)
+{
+	return TSharedPtr<FPixelCaptureCapturerRHIRDG>(new FPixelCaptureCapturerRHIRDG(InScale));
+}
+
+FPixelCaptureCapturerRHIRDG::FPixelCaptureCapturerRHIRDG(float InScale)
+	: Scale(InScale)
+{
+}
+
+IPixelCaptureOutputFrame* FPixelCaptureCapturerRHIRDG::CreateOutputBuffer(int32 InputWidth, int32 InputHeight)
+{
+	const int32 Width = InputWidth * Scale;
+	const int32 Height = InputHeight * Scale;
+
+	FRHITextureCreateDesc TextureDesc =
+		FRHITextureCreateDesc::Create2D(TEXT("FFrameDataH264 Texture"), Width, Height, EPixelFormat::PF_B8G8R8A8)
+			.SetClearValue(FClearValueBinding::None)
+			.SetFlags(ETextureCreateFlags::RenderTargetable)
+			.SetInitialState(ERHIAccess::Present)
+			.DetermineInititialState();
+
+	if (RHIGetInterfaceType() == ERHIInterfaceType::Vulkan)
+	{
+		TextureDesc.AddFlags(ETextureCreateFlags::External);
+	}
+	else
+	{
+		TextureDesc.AddFlags(ETextureCreateFlags::Shared);
+	}
+
+	return new FPixelCaptureOutputFrameRHI(GDynamicRHI->RHICreateTexture(TextureDesc));
+}
+
+void FPixelCaptureCapturerRHIRDG::BeginProcess(const IPixelCaptureInputFrame& InputFrame, IPixelCaptureOutputFrame* OutputBuffer)
+{
+	checkf(InputFrame.GetType() == StaticCast<int32>(PixelCaptureBufferFormat::FORMAT_RHI), TEXT("Incorrect source frame coming into frame capture process."));
+
+	MarkCPUWorkStart();
+
+	FRHICommandListImmediate& RHICmdList = FRHICommandListExecutor::GetImmediateCommandList();
+
+	MarkGPUWorkStart();
+	const FPixelCaptureInputFrameRHI& RHISourceFrame = StaticCast<const FPixelCaptureInputFrameRHI&>(InputFrame);
+	FPixelCaptureOutputFrameRHI* OutputH264Buffer = StaticCast<FPixelCaptureOutputFrameRHI*>(OutputBuffer);
+	CopyTextureRDG(RHICmdList, RHISourceFrame.FrameTexture, OutputH264Buffer->GetFrameTexture());
+
+	OnRHIStageComplete();
+}
+
+void FPixelCaptureCapturerRHIRDG::CheckComplete()
+{
+
+}
+
+void FPixelCaptureCapturerRHIRDG::OnRHIStageComplete()
+{
+	MarkGPUWorkEnd();
+	EndProcess();
+}
diff --git a/Engine/Plugins/Media/PixelCapture/Source/PixelCapture/Private/PixelCaptureUtils.h b/Engine/Plugins/Media/PixelCapture/Source/PixelCapture/Private/PixelCaptureUtils.h
--- a/Engine/Plugins/Media/PixelCapture/Source/PixelCapture/Private/PixelCaptureUtils.h
+++ b/Engine/Plugins/Media/PixelCapture/Source/PixelCapture/Private/PixelCaptureUtils.h
@@ -6,6 +6,7 @@
 #include "RHI.h"
 #include "Runtime/Renderer/Private/ScreenPass.h"
 #include "ScreenRendering.h"
+#include "MediaShaders.h"
 
 /*
     * Copy from one texture to another.
@@ -76,4 +77,64 @@ inline void CopyTexture(FRHICommandList& RHICmdList, FTextureRHIRef SourceTextur
     {
         RHICmdList.WriteGPUFence(Fence);
     }
+}
+
+inline void CopyTextureRDG(FRHICommandListImmediate& RHICmdList, FTextureRHIRef SourceTexture, FTextureRHIRef DestTexture)
+{
+	FRDGBuilder GraphBuilder(RHICmdList);
+
+	// Register an external RDG rexture from the source texture
+	FRDGTexture* InputTexture = GraphBuilder.RegisterExternalTexture(CreateRenderTarget(SourceTexture, TEXT("PixelCaptureCopySourceTexture")));
+
+	// Register an external RDG texture from the output buffer
+	FRDGTexture* OutputTexture = GraphBuilder.RegisterExternalTexture(CreateRenderTarget(DestTexture, TEXT("PixelCaptureCopyDestTexture")));
+
+	if (InputTexture->Desc.Format == OutputTexture->Desc.Format &&
+        InputTexture->Desc.Extent.X == OutputTexture->Desc.Extent.X &&
+        InputTexture->Desc.Extent.Y == OutputTexture->Desc.Extent.Y)
+    {
+		// The formats are the same and size are the same. simple copy
+		AddDrawTexturePass(
+			GraphBuilder,
+			GetGlobalShaderMap(GMaxRHIFeatureLevel),
+			InputTexture,
+			OutputTexture,
+			FRDGDrawTextureInfo()
+		);
+	}
+	else
+	{
+		// The formats or size differ to pixel shader stuff
+		//Configure source/output viewport to get the right UV scaling from source texture to output texture
+		FScreenPassTextureViewport InputViewport(InputTexture);
+		FScreenPassTextureViewport OutputViewport(OutputTexture);
+
+		FGlobalShaderMap* GlobalShaderMap = GetGlobalShaderMap(GMaxRHIFeatureLevel);
+		TShaderMapRef<FScreenPassVS> VertexShader(GlobalShaderMap);
+
+		// In cases where texture is converted from a format that doesn't have A channel, we want to force set it to 1.
+		int32 ConversionOperation = 0; // None
+		FModifyAlphaSwizzleRgbaPS::FPermutationDomain PermutationVector;
+		PermutationVector.Set<FModifyAlphaSwizzleRgbaPS::FConversionOp>(ConversionOperation);
+
+		// Rectangle area to use from source
+		const FIntRect ViewRect(FIntPoint(0, 0), InputTexture->Desc.Extent);
+
+		//Dummy ViewFamily/ViewInfo created to use built in Draw Screen/Texture Pass
+		FSceneViewFamily ViewFamily(FSceneViewFamily::ConstructionValues(nullptr, nullptr, FEngineShowFlags(ESFIM_Game))
+			.SetTime(FGameTime())
+			.SetGammaCorrection(1.0f));
+		FSceneViewInitOptions ViewInitOptions;
+		ViewInitOptions.ViewFamily = &ViewFamily;
+		ViewInitOptions.SetViewRectangle(ViewRect);
+		ViewInitOptions.ViewOrigin = FVector::ZeroVector;
+		ViewInitOptions.ViewRotationMatrix = FMatrix::Identity;
+		ViewInitOptions.ProjectionMatrix = FMatrix::Identity;
+		FViewInfo ViewInfo = FViewInfo(ViewInitOptions);
+
+		TShaderMapRef<FModifyAlphaSwizzleRgbaPS> PixelShader(GlobalShaderMap, PermutationVector);
+		FModifyAlphaSwizzleRgbaPS::FParameters* Parameters = PixelShader->AllocateAndSetParameters(GraphBuilder, InputTexture, OutputTexture);
+		AddDrawScreenPass(GraphBuilder, RDG_EVENT_NAME("PixelCapturerSwizzle"), ViewInfo, OutputViewport, InputViewport, VertexShader, PixelShader, Parameters);
+	}
+	GraphBuilder.Execute();
 }
\ No newline at end of file
diff --git a/Engine/Plugins/Media/PixelCapture/Source/PixelCapture/Public/PixelCaptureCapturerRHIRDG.h b/Engine/Plugins/Media/PixelCapture/Source/PixelCapture/Public/PixelCaptureCapturerRHIRDG.h
new file mode 100644
--- /dev/null
+++ b/Engine/Plugins/Media/PixelCapture/Source/PixelCapture/Public/PixelCaptureCapturerRHIRDG.h
@@ -0,0 +1,34 @@
+// Copyright Epic Games, Inc. All Rights Reserved.
+
+#pragma once
+
+#include "PixelCaptureCapturer.h"
+#include "RHI.h"
+
+/**
+ * A basic capturer that will copy RHI texture frames.
+ * Input: FPixelCaptureInputFrameRHI
+ * Output: FPixelCaptureOutputFrameRHI
+ */
+class PIXELCAPTURE_API FPixelCaptureCapturerRHIRDG : public FPixelCaptureCapturer, public TSharedFromThis<FPixelCaptureCapturerRHIRDG>
+{
+public:
+	/**
+	 * Creates a new Capturer capturing the input frame at the given scale.
+	 * @param InScale The scale of the resulting output capture.
+	 */
+	static TSharedPtr<FPixelCaptureCapturerRHIRDG> Create(float InScale);
+	virtual ~FPixelCaptureCapturerRHIRDG() = default;
+
+protected:
+	virtual FString GetCapturerName() const override { return "RHI RDG Copy"; }
+	virtual IPixelCaptureOutputFrame* CreateOutputBuffer(int32 InputWidth, int32 InputHeight) override;
+	virtual void BeginProcess(const IPixelCaptureInputFrame& InputFrame, IPixelCaptureOutputFrame* OutputBuffer) override;
+
+private:
+	float Scale = 1.0f;
+
+	FPixelCaptureCapturerRHIRDG(float InScale);
+	void CheckComplete();
+	void OnRHIStageComplete();
+};
diff --git a/Engine/Plugins/Media/PixelStreaming/Source/PixelStreaming/Private/PixelStreamingVideoInputRHI.cpp b/Engine/Plugins/Media/PixelStreaming/Source/PixelStreaming/Private/PixelStreamingVideoInputRHI.cpp
--- a/Engine/Plugins/Media/PixelStreaming/Source/PixelStreaming/Private/PixelStreamingVideoInputRHI.cpp
+++ b/Engine/Plugins/Media/PixelStreaming/Source/PixelStreaming/Private/PixelStreamingVideoInputRHI.cpp
@@ -3,9 +3,9 @@
 #include "PixelStreamingVideoInputRHI.h"
 #include "PixelStreamingPrivate.h"
 #include "Settings.h"
-
 #include "PixelCaptureBufferFormat.h"
 #include "PixelCaptureCapturerRHI.h"
+#include "PixelCaptureCapturerRHIRDG.h"
 #include "PixelCaptureCapturerRHIToI420CPU.h"
 #include "PixelCaptureCapturerRHIToI420Compute.h"
 
@@ -15,7 +15,17 @@ TSharedPtr<FPixelCaptureCapturer> FPixelStreamingVideoInputRHI::CreateCapturer(i
 	{
 		case PixelCaptureBufferFormat::FORMAT_RHI:
 		{
-			return FPixelCaptureCapturerRHI::Create(FinalScale);
+			// "Safe Texture Copy" polls a fence to ensure a GPU copy is complete
+			// the RDG pathway does not poll a fence so is more unsafe but offers 
+			// a significant performance increase
+			if (UE::PixelStreaming::Settings::IsUsingSafeTextureCopy())
+			{
+				return FPixelCaptureCapturerRHI::Create(FinalScale);
+			}
+			else
+			{
+				return FPixelCaptureCapturerRHIRDG::Create(FinalScale);
+			}
 		}
 		case PixelCaptureBufferFormat::FORMAT_I420:
 		{
diff --git a/Engine/Plugins/Media/PixelStreaming/Source/PixelStreaming/Private/Settings.cpp b/Engine/Plugins/Media/PixelStreaming/Source/PixelStreaming/Private/Settings.cpp
--- a/Engine/Plugins/Media/PixelStreaming/Source/PixelStreaming/Private/Settings.cpp
+++ b/Engine/Plugins/Media/PixelStreaming/Source/PixelStreaming/Private/Settings.cpp
@@ -238,23 +238,9 @@ namespace UE::PixelStreaming::Settings
 		ECVF_Default);
 
 	TArray<FKey> FilteredKeys;
-	TArray<FName> ActiveTextureSourceTypes;
-	bool bDecoupleFrameRate = true;
 	FString DefaultStreamerID = TEXT("DefaultStreamer");
 	FString DefaultSignallingURL = TEXT("ws://127.0.0.1:8888");
 
-	void OnActiveTextureSourcesChanged(IConsoleVariable* Var)
-	{
-		FString CommaList = Var->GetString();
-		TArray<FString> StringArray;
-		CommaList.ParseIntoArray(StringArray, TEXT(","), true);
-		ActiveTextureSourceTypes.Empty();
-		for (auto&& TextureSourceTypeStr : StringArray)
-		{
-			ActiveTextureSourceTypes.Add(FName(TextureSourceTypeStr));
-		}
-	}
-
 	void OnFilteredKeysChanged(IConsoleVariable* Var)
 	{
 		FString CommaList = Var->GetString();
@@ -386,17 +372,6 @@ namespace UE::PixelStreaming::Settings
 			return AVEncoder::FVideoEncoder::H264Profile::BASELINE;
 		return Iter->second;
 	}
-
-	const TArray<FName>& GetActiveTextureSourceTypes()
-	{
-		return ActiveTextureSourceTypes;
-	}
-
-	void SetActiveTextureSourceTypes(TArray<FName> InActiveTextureSourceTypes)
-	{
-		ActiveTextureSourceTypes = InActiveTextureSourceTypes;
-	}
-
 	// End utility functions etc.
 
 	FSimulcastParameters SimulcastParameters;
@@ -478,11 +453,6 @@ namespace UE::PixelStreaming::Settings
 		}
 	}
 
-	bool DecoupleFrameRate()
-	{
-		return bDecoupleFrameRate;
-	}
-
 	bool IsCodecVPX()
 	{
 		EPixelStreamingCodec SelectedCodec = GetSelectedCodec();
@@ -561,16 +531,12 @@ namespace UE::PixelStreaming::Settings
 		UE_LOG(LogPixelStreaming, Log, TEXT("Initialising Pixel Streaming settings."));
 
 		FString StringOptions;
-		bDecoupleFrameRate = !FParse::Value(FCommandLine::Get(), TEXT("PixelStreamingDisableDecoupleFrameRate"), StringOptions, false);
 
 		FParse::Value(FCommandLine::Get(), TEXT("PixelStreamingID="), DefaultStreamerID, false);
 
 		CVarPixelStreamingOnScreenStats.AsVariable()->SetOnChangedCallback(FConsoleVariableDelegate::CreateStatic(&OnHudStatsToggled));
 		CVarPixelStreamingKeyFilter.AsVariable()->SetOnChangedCallback(FConsoleVariableDelegate::CreateStatic(&OnFilteredKeysChanged));
-		CVarPixelStreamingVideoTracks.AsVariable()->SetOnChangedCallback(FConsoleVariableDelegate::CreateStatic(&OnActiveTextureSourcesChanged));
 		CVarPixelStreamingWebRTCFps.AsVariable()->SetOnChangedCallback(FConsoleVariableDelegate::CreateStatic(&OnWebRTCFpsChanged));
-		OnActiveTextureSourcesChanged(CVarPixelStreamingVideoTracks.AsVariable());
-
 		CVarPixelStreamingEncoderKeyframeInterval.AsVariable()->SetOnChangedCallback(FConsoleVariableDelegate::CreateStatic(&OnKeyframeIntervalChanged));
 
 		// Values parse from commands line
diff --git a/Engine/Plugins/Media/PixelStreaming/Source/PixelStreaming/Private/Settings.h b/Engine/Plugins/Media/PixelStreaming/Source/PixelStreaming/Private/Settings.h
--- a/Engine/Plugins/Media/PixelStreaming/Source/PixelStreaming/Private/Settings.h
+++ b/Engine/Plugins/Media/PixelStreaming/Source/PixelStreaming/Private/Settings.h
@@ -74,7 +74,6 @@ namespace UE::PixelStreaming::Settings
 	};
 
 	// Begin utility functions etc.
-	bool DecoupleFrameRate();
 	bool IsCodecVPX();
 	void SetCodec(EPixelStreamingCodec Codec);
 	EPixelStreamingCodec GetSelectedCodec();
@@ -82,8 +81,6 @@ namespace UE::PixelStreaming::Settings
 	AVEncoder::FVideoEncoder::MultipassMode GetMultipassCVar();
 	webrtc::DegradationPreference GetDegradationPreference();
 	AVEncoder::FVideoEncoder::H264Profile GetH264Profile();
-	const TArray<FName>& GetActiveTextureSourceTypes();
-	void SetActiveTextureSourceTypes(TArray<FName> InActiveTextureSourceTypes);
 	EInputControllerMode GetInputControllerMode();
 	FString GetDefaultStreamerID();
 	FString GetDefaultSignallingURL();
@@ -114,6 +111,16 @@ namespace UE::PixelStreaming::Settings
 		return FParse::Param(FCommandLine::Get(), TEXT("PixelStreamingHideCursor"));
 	}
 
+	inline bool IsUsingSafeTextureCopy()
+	{
+		return FParse::Param(FCommandLine::Get(), TEXT("PixelCaptureUseFence"));
+	}
+
+	inline bool CoupleFrameRate()
+	{
+		return !FParse::Param(FCommandLine::Get(), TEXT("PixelStreamingDecoupleFrameRate"));
+	}
+
 	inline bool GetSignallingServerUrl(FString& OutSignallingServerURL)
 	{
 		return FParse::Value(FCommandLine::Get(), TEXT("PixelStreamingURL="), OutSignallingServerURL);
diff --git a/Engine/Plugins/Media/PixelStreaming/Source/PixelStreaming/Private/VideoSourceGroup.cpp b/Engine/Plugins/Media/PixelStreaming/Source/PixelStreaming/Private/VideoSourceGroup.cpp
--- a/Engine/Plugins/Media/PixelStreaming/Source/PixelStreaming/Private/VideoSourceGroup.cpp
+++ b/Engine/Plugins/Media/PixelStreaming/Source/PixelStreaming/Private/VideoSourceGroup.cpp
@@ -13,7 +13,7 @@ namespace UE::PixelStreaming
 	}
 
 	FVideoSourceGroup::FVideoSourceGroup()
-		: bCoupleFramerate(!Settings::DecoupleFrameRate())
+		: bCoupleFramerate(Settings::CoupleFrameRate() || !Settings::IsUsingSafeTextureCopy())
 		, FramesPerSecond(Settings::CVarPixelStreamingWebRTCFps.GetValueOnAnyThread())
 	{
 	}
