Run Pixel Capture fence polling on background thread to prevent deadlock issues

diff --git a/Engine/Plugins/Media/PixelCapture/Source/PixelCapture/Private/PixelCaptureCapturerRHI.cpp b/Engine/Plugins/Media/PixelCapture/Source/PixelCapture/Private/PixelCaptureCapturerRHI.cpp
--- a/Engine/Plugins/Media/PixelCapture/Source/PixelCapture/Private/PixelCaptureCapturerRHI.cpp
+++ b/Engine/Plugins/Media/PixelCapture/Source/PixelCapture/Private/PixelCaptureCapturerRHI.cpp
@@ -69,7 +69,7 @@ void FPixelCaptureCapturerRHI::CheckComplete()
 	// in lieu of a proper callback we need to capture a thread to poll the fence
 	// so we know as quickly as possible when we can readback.
 	TSharedRef<FPixelCaptureCapturerRHI> ThisRHIRef = StaticCastSharedRef<FPixelCaptureCapturerRHI>(AsShared());
-	AsyncTask(ENamedThreads::AnyHiPriThreadHiPriTask, [ThisRHIRef]() {
+	AsyncTask(ENamedThreads::AnyBackgroundHiPriTask, [ThisRHIRef]() {
 		while (!ThisRHIRef->Fence->Poll())
 		{
 			// we want to check quickly but we dont want to burn CPU.
