Updating the copy texture routine because ASAN was showing stack usage after it was deleted.

diff --git a/Engine/Plugins/Media/PixelCapture/Source/PixelCapture/Private/PixelCaptureUtils.h b/Engine/Plugins/Media/PixelCapture/Source/PixelCapture/Private/PixelCaptureUtils.h
--- a/Engine/Plugins/Media/PixelCapture/Source/PixelCapture/Private/PixelCaptureUtils.h
+++ b/Engine/Plugins/Media/PixelCapture/Source/PixelCapture/Private/PixelCaptureUtils.h
@@ -130,7 +130,14 @@ inline void CopyTextureRDG(FRHICommandListImmediate& RHICmdList, FTextureRHIRef
 		ViewInitOptions.ViewOrigin = FVector::ZeroVector;
 		ViewInitOptions.ViewRotationMatrix = FMatrix::Identity;
 		ViewInitOptions.ProjectionMatrix = FMatrix::Identity;
-		FViewInfo ViewInfo = FViewInfo(ViewInitOptions);
+		// the RDG seems to take this parameter by reference all the way through which means
+		// since we're submitting async work here that this struct might not be around when
+		// the work is done. ASAN was showing stack access to a destroyed object that pointed
+		// to this struct. making it static is kind of a nasty hack but it guarantees that it
+		// wont be accessing dead memory in the RDG.
+		// TODO (matthew.cotton) this probably requires rethinking because it adds hidden global state
+		static FViewInfo ViewInfo;
+		ViewInfo = FViewInfo(ViewInitOptions);
 
 		TShaderMapRef<FModifyAlphaSwizzleRgbaPS> PixelShader(GlobalShaderMap, PermutationVector);
 		FModifyAlphaSwizzleRgbaPS::FParameters* Parameters = PixelShader->AllocateAndSetParameters(GraphBuilder, InputTexture, OutputTexture);
