Updating the handling of the FViewInfo object in the CopyTextureRDG function.

diff --git a/Engine/Plugins/Media/PixelCapture/Source/PixelCapture/Private/PixelCaptureUtils.h b/Engine/Plugins/Media/PixelCapture/Source/PixelCapture/Private/PixelCaptureUtils.h
--- a/Engine/Plugins/Media/PixelCapture/Source/PixelCapture/Private/PixelCaptureUtils.h
+++ b/Engine/Plugins/Media/PixelCapture/Source/PixelCapture/Private/PixelCaptureUtils.h
@@ -104,6 +104,16 @@ inline void CopyTextureRDG(FRHICommandListImmediate& RHICmdList, FTextureRHIRef
 	}
 	else
 	{
+		// TODO (matthew.cotton) A lot of this was copied from AddScreenPass and altered for our use.
+		// This was altered because the FViewInfo we create here on the stack is added to the RHI
+		// queue and is accessed possibly after this stack frame was released. Originally we tried
+		// just changing it to static so we could reuse it, but the object does not allow copying.
+		// So the current implementation creates the View on the heap using a shared ptr that is then
+		// captured in a lambda for the RHI queue. This allows us to update the view on calls to this
+		// function and not worry about it's lifetime.
+		// The todo is here because there has to be a better way to achieve what we want without
+		// all of this song and dance.
+
 		// The formats or size differ to pixel shader stuff
 		//Configure source/output viewport to get the right UV scaling from source texture to output texture
 		FScreenPassTextureViewport InputViewport(InputTexture);
@@ -130,18 +140,29 @@ inline void CopyTextureRDG(FRHICommandListImmediate& RHICmdList, FTextureRHIRef
 		ViewInitOptions.ViewOrigin = FVector::ZeroVector;
 		ViewInitOptions.ViewRotationMatrix = FMatrix::Identity;
 		ViewInitOptions.ProjectionMatrix = FMatrix::Identity;
-		// the RDG seems to take this parameter by reference all the way through which means
-		// since we're submitting async work here that this struct might not be around when
-		// the work is done. ASAN was showing stack access to a destroyed object that pointed
-		// to this struct. making it static is kind of a nasty hack but it guarantees that it
-		// wont be accessing dead memory in the RDG.
-		// TODO (matthew.cotton) this probably requires rethinking because it adds hidden global state
-		static FViewInfo ViewInfo;
-		ViewInfo = FViewInfo(ViewInitOptions);
+		TSharedPtr<FViewInfo> View = MakeShared<FViewInfo>(ViewInitOptions);
 
 		TShaderMapRef<FModifyAlphaSwizzleRgbaPS> PixelShader(GlobalShaderMap, PermutationVector);
-		FModifyAlphaSwizzleRgbaPS::FParameters* Parameters = PixelShader->AllocateAndSetParameters(GraphBuilder, InputTexture, OutputTexture);
-		AddDrawScreenPass(GraphBuilder, RDG_EVENT_NAME("PixelCapturerSwizzle"), ViewInfo, OutputViewport, InputViewport, VertexShader, PixelShader, Parameters);
+		FModifyAlphaSwizzleRgbaPS::FParameters* PixelShaderParameters = PixelShader->AllocateAndSetParameters(GraphBuilder, InputTexture, OutputTexture);
+		
+		FRHIBlendState* BlendState = FScreenPassPipelineState::FDefaultBlendState::GetRHI();
+		FRHIDepthStencilState* DepthStencilState = FScreenPassPipelineState::FDefaultDepthStencilState::GetRHI();
+
+		ClearUnusedGraphResources(PixelShader, PixelShaderParameters);
+
+		const FScreenPassPipelineState PipelineState(VertexShader, PixelShader, BlendState, DepthStencilState);
+
+		GraphBuilder.AddPass(
+			RDG_EVENT_NAME("PixelCapturerSwizzle"),
+			PixelShaderParameters,
+			ERDGPassFlags::Raster,
+			[View, OutputViewport, InputViewport, PipelineState, PixelShader, PixelShaderParameters](FRHICommandList& RHICmdList)
+		{
+			DrawScreenPass(RHICmdList, *View, OutputViewport, InputViewport, PipelineState, EScreenPassDrawFlags::None, [&](FRHICommandList&)
+			{
+				SetShaderParameters(RHICmdList, PixelShader, PixelShader.GetPixelShader(), *PixelShaderParameters);
+			});
+		});
 	}
 	GraphBuilder.Execute();
 }
