Fix deadlock issue in Pixel Streaming that would occur after an intermittent amount of time causing the application to freeze and crash.

diff --git a/Engine/Plugins/Media/PixelCapture/Source/PixelCapture/Private/PixelCaptureCapturerMultiFormat.cpp b/Engine/Plugins/Media/PixelCapture/Source/PixelCapture/Private/PixelCaptureCapturerMultiFormat.cpp
--- a/Engine/Plugins/Media/PixelCapture/Source/PixelCapture/Private/PixelCaptureCapturerMultiFormat.cpp
+++ b/Engine/Plugins/Media/PixelCapture/Source/PixelCapture/Private/PixelCaptureCapturerMultiFormat.cpp
@@ -56,15 +56,17 @@ void FPixelCaptureCapturerMultiFormat::Capture(const IPixelCaptureInputFrame& So
 
 	if (FormatCapturers.Num() > 0)
 	{
+		// UE-173694: Need to acquire the FormatGuard before FPixelCaptureCapturerLayered acquires its LayersGuard
+		// to avoid a deadlock with the EncoderQueue
+		FScopeLock FormatLock(&FormatGuard);
+		// iterate a temp copy so modifications to the map in Capture
+		// (and their callback events) is ok.
+		// Note: When considering whether the lifetime of this lock can be shorter consider the deadlock issue we had in UE-173694
+
 		TArray<TSharedPtr<FPixelCaptureCapturerLayered>> TempCopy;
+		for (auto& FormatCapturer : FormatCapturers)
 		{
-			// iterate a temp copy so modifications to the map in Capture
-			// (and their callback events) is ok.
-			FScopeLock FormatLock(&FormatGuard);
-			for (auto& FormatCapturer : FormatCapturers)
-			{
-				TempCopy.Add(FormatCapturer.Value);
-			}
+			TempCopy.Add(FormatCapturer.Value);
 		}
 		// start all the format captures
 		PendingFormats = FormatCapturers.Num();
