Remove direct use of GDynamicRHI pointer.

diff --git a/Engine/Plugins/Media/PixelCapture/Source/PixelCapture/Private/PixelCaptureCapturerRHIToI420Compute.cpp b/Engine/Plugins/Media/PixelCapture/Source/PixelCapture/Private/PixelCaptureCapturerRHIToI420Compute.cpp
--- a/Engine/Plugins/Media/PixelCapture/Source/PixelCapture/Private/PixelCaptureCapturerRHIToI420Compute.cpp
+++ b/Engine/Plugins/Media/PixelCapture/Source/PixelCapture/Private/PixelCaptureCapturerRHIToI420Compute.cpp
@@ -58,9 +58,9 @@ void FPixelCaptureCapturerRHIToI420Compute::Initialize(int32 InputWidth, int32 I
 			.SetInitialState(ERHIAccess::UAVCompute)
 			.DetermineInititialState();
 
-	TextureY = GDynamicRHI->RHICreateTexture(TextureDescY);
-	TextureU = GDynamicRHI->RHICreateTexture(TextureDescUV);
-	TextureV = GDynamicRHI->RHICreateTexture(TextureDescUV);
+	TextureY = RHICreateTexture(TextureDescY);
+	TextureU = RHICreateTexture(TextureDescUV);
+	TextureV = RHICreateTexture(TextureDescUV);
 
 	FRHITextureCreateDesc StagingDescY =
 		FRHITextureCreateDesc::Create2D(TEXT("YUV Output CPU Texture"), PlaneYDimensions.X, PlaneYDimensions.Y, EPixelFormat::PF_R8)
@@ -76,20 +76,22 @@ void FPixelCaptureCapturerRHIToI420Compute::Initialize(int32 InputWidth, int32 I
 			.SetInitialState(ERHIAccess::Unknown)
 			.DetermineInititialState();
 
-	StagingTextureY = GDynamicRHI->RHICreateTexture(StagingDescY);
-	StagingTextureU = GDynamicRHI->RHICreateTexture(StagingDescUV);
-	StagingTextureV = GDynamicRHI->RHICreateTexture(StagingDescUV);
+	StagingTextureY = RHICreateTexture(StagingDescY);
+	StagingTextureU = RHICreateTexture(StagingDescUV);
+	StagingTextureV = RHICreateTexture(StagingDescUV);
 
-	TextureYUAV = GDynamicRHI->RHICreateUnorderedAccessView(TextureY, 0, 0, 0);
-	TextureUUAV = GDynamicRHI->RHICreateUnorderedAccessView(TextureU, 0, 0, 0);
-	TextureVUAV = GDynamicRHI->RHICreateUnorderedAccessView(TextureV, 0, 0, 0);
+	TextureYUAV = RHICreateUnorderedAccessView(TextureY);
+	TextureUUAV = RHICreateUnorderedAccessView(TextureU);
+	TextureVUAV = RHICreateUnorderedAccessView(TextureV);
+
+	FRHICommandListImmediate& RHICmdList = FRHICommandListExecutor::GetImmediateCommandList();
 
 	int32 OutWidth, OutHeight;
-	GDynamicRHI->RHIMapStagingSurface(StagingTextureY, nullptr, MappedY, OutWidth, OutHeight);
+	RHICmdList.MapStagingSurface(StagingTextureY, nullptr, MappedY, OutWidth, OutHeight);
 	YStride = OutWidth;
-	GDynamicRHI->RHIMapStagingSurface(StagingTextureU, nullptr, MappedU, OutWidth, OutHeight);
+	RHICmdList.MapStagingSurface(StagingTextureU, nullptr, MappedU, OutWidth, OutHeight);
 	UStride = OutWidth;
-	GDynamicRHI->RHIMapStagingSurface(StagingTextureV, nullptr, MappedV, OutWidth, OutHeight);
+	RHICmdList.MapStagingSurface(StagingTextureV, nullptr, MappedV, OutWidth, OutHeight);
 	VStride = OutWidth;
 
 	FPixelCaptureCapturer::Initialize(InputWidth, InputHeight);
@@ -155,9 +157,10 @@ void FPixelCaptureCapturerRHIToI420Compute::OnRHIStageComplete(IPixelCaptureOutp
 
 void FPixelCaptureCapturerRHIToI420Compute::CleanUp()
 {
-	GDynamicRHI->RHIUnmapStagingSurface(StagingTextureY);
-	GDynamicRHI->RHIUnmapStagingSurface(StagingTextureU);
-	GDynamicRHI->RHIUnmapStagingSurface(StagingTextureV);
+	FRHICommandListImmediate& RHICmdList = FRHICommandListExecutor::GetImmediateCommandList();
+	RHICmdList.UnmapStagingSurface(StagingTextureY);
+	RHICmdList.UnmapStagingSurface(StagingTextureU);
+	RHICmdList.UnmapStagingSurface(StagingTextureV);
 
 	MappedY = nullptr;
 	MappedU = nullptr;
