Fixing out of order behaviour leading to a crash when running with -onethread and encoding VPX
- Changed the order of end of CPU timing marks to run before submitting the finish lambda to the renderer. This means that the submission process doesn't get counted as CPU work but it's timing should be very minor.

diff --git a/Engine/Plugins/Media/PixelCapture/Source/PixelCapture/Private/PixelCaptureCapturer.cpp b/Engine/Plugins/Media/PixelCapture/Source/PixelCapture/Private/PixelCaptureCapturer.cpp
--- a/Engine/Plugins/Media/PixelCapture/Source/PixelCapture/Private/PixelCaptureCapturer.cpp
+++ b/Engine/Plugins/Media/PixelCapture/Source/PixelCapture/Private/PixelCaptureCapturer.cpp
@@ -73,6 +73,7 @@ void FPixelCaptureCapturer::MarkCPUWorkStart()
 
 void FPixelCaptureCapturer::MarkCPUWorkEnd()
 {
+	check(CurrentOutputBuffer != nullptr);
 	CurrentOutputBuffer->Metadata.CaptureProcessCPUTime += FPlatformTime::Cycles64() - CPUStartTime;
 	CPUStartTime = 0;
 
