Fixing a deadlock that could happen when using dx11 on AWS instances.
- On startup we generate a freezeframe using RHI. During that capture we loop poll on a thread waiting for the gpu fence to complete but doing so seems to create a deadlock where the game thread and others appear to be waiting for that one looping thread. I couldn't lock down the exact reasons why however.
- This change will change the behaviour of dx11 captures where the thread will check the fence and if it is not complete we submit a new task, rather than looping and locking the thread.

diff --git a/Engine/Plugins/Media/PixelCapture/Source/PixelCapture/Private/PixelCaptureCapturerRHI.cpp b/Engine/Plugins/Media/PixelCapture/Source/PixelCapture/Private/PixelCaptureCapturerRHI.cpp
--- a/Engine/Plugins/Media/PixelCapture/Source/PixelCapture/Private/PixelCaptureCapturerRHI.cpp
+++ b/Engine/Plugins/Media/PixelCapture/Source/PixelCapture/Private/PixelCaptureCapturerRHI.cpp
@@ -17,6 +17,7 @@ FPixelCaptureCapturerRHI::FPixelCaptureCapturerRHI(float InScale)
 	: Scale(InScale)
 {
 	Fence = GDynamicRHI->RHICreateGPUFence(TEXT("FPixelCaptureCapturerRHI Fence"));
+	RHIType = GDynamicRHI ? RHIGetInterfaceType() : ERHIInterfaceType::Hidden;
 }
 
 IPixelCaptureOutputFrame* FPixelCaptureCapturerRHI::CreateOutputBuffer(int32 InputWidth, int32 InputHeight)
@@ -68,15 +69,34 @@ void FPixelCaptureCapturerRHI::CheckComplete()
 {
 	// in lieu of a proper callback we need to capture a thread to poll the fence
 	// so we know as quickly as possible when we can readback.
-	TSharedRef<FPixelCaptureCapturerRHI> ThisRHIRef = StaticCastSharedRef<FPixelCaptureCapturerRHI>(AsShared());
-	AsyncTask(ENamedThreads::AnyBackgroundHiPriTask, [ThisRHIRef]() {
-		while (!ThisRHIRef->Fence->Poll())
+
+	if (RHIType == ERHIInterfaceType::D3D11)
+	{
+		// sometimes on dx11 we end up in a deadlock when we loop here polling the fence
+		// so instead we check and then submit a new check task.
+		if (!Fence->Poll())
+		{
+			TSharedRef<FPixelCaptureCapturerRHI> ThisRHIRef = StaticCastSharedRef<FPixelCaptureCapturerRHI>(AsShared());
+			AsyncTask(ENamedThreads::AnyBackgroundHiPriTask, [ThisRHIRef]() { ThisRHIRef->CheckComplete(); });
+		}
+		else
 		{
-			// we want to check quickly but we dont want to burn CPU.
-			FPlatformProcess::Sleep(0.0001f);
+			OnRHIStageComplete();
 		}
-		ThisRHIRef->OnRHIStageComplete();
-	});
+	}
+	else
+	{
+		// we submit a task here to loop poll the fence so that we get notified as quickly as possible about its completion.
+		TSharedRef<FPixelCaptureCapturerRHI> ThisRHIRef = StaticCastSharedRef<FPixelCaptureCapturerRHI>(AsShared());
+		AsyncTask(ENamedThreads::AnyBackgroundHiPriTask, [ThisRHIRef]() {
+			while (!ThisRHIRef->Fence->Poll())
+			{
+				// we want to check quickly but we dont want to burn CPU.
+				FPlatformProcess::Sleep(0.0001f);
+			}
+			ThisRHIRef->OnRHIStageComplete();
+		});
+	}
 }
 
 void FPixelCaptureCapturerRHI::OnRHIStageComplete()
diff --git a/Engine/Plugins/Media/PixelCapture/Source/PixelCapture/Public/PixelCaptureCapturerRHI.h b/Engine/Plugins/Media/PixelCapture/Source/PixelCapture/Public/PixelCaptureCapturerRHI.h
--- a/Engine/Plugins/Media/PixelCapture/Source/PixelCapture/Public/PixelCaptureCapturerRHI.h
+++ b/Engine/Plugins/Media/PixelCapture/Source/PixelCapture/Public/PixelCaptureCapturerRHI.h
@@ -29,6 +29,7 @@ private:
 	float Scale = 1.0f;
 
 	FGPUFenceRHIRef Fence;
+	ERHIInterfaceType RHIType;
 
 	FPixelCaptureCapturerRHI(float InScale);
 	void CheckComplete();
