Removing unnecessary check triggering a static analysis issue.
- Reverting back to GDynamicRHI not being checked.
- If GDynamicRHI is null when our object is created we have larger issues and no clean way to back out.

diff --git a/Engine/Plugins/Media/PixelCapture/Source/PixelCapture/Private/PixelCaptureCapturerRHI.cpp b/Engine/Plugins/Media/PixelCapture/Source/PixelCapture/Private/PixelCaptureCapturerRHI.cpp
--- a/Engine/Plugins/Media/PixelCapture/Source/PixelCapture/Private/PixelCaptureCapturerRHI.cpp
+++ b/Engine/Plugins/Media/PixelCapture/Source/PixelCapture/Private/PixelCaptureCapturerRHI.cpp
@@ -17,7 +17,7 @@ FPixelCaptureCapturerRHI::FPixelCaptureCapturerRHI(float InScale)
 	: Scale(InScale)
 {
 	Fence = GDynamicRHI->RHICreateGPUFence(TEXT("FPixelCaptureCapturerRHI Fence"));
-	RHIType = GDynamicRHI ? RHIGetInterfaceType() : ERHIInterfaceType::Hidden;
+	RHIType = RHIGetInterfaceType();
 }
 
 IPixelCaptureOutputFrame* FPixelCaptureCapturerRHI::CreateOutputBuffer(int32 InputWidth, int32 InputHeight)
