Fix: PixelStreamed applications crash on AWS and CoreWeave Linux containers

diff --git a/Engine/Plugins/Media/PixelCapture/Source/PixelCapture/Private/PixelCaptureCapturerRHI.cpp b/Engine/Plugins/Media/PixelCapture/Source/PixelCapture/Private/PixelCaptureCapturerRHI.cpp
--- a/Engine/Plugins/Media/PixelCapture/Source/PixelCapture/Private/PixelCaptureCapturerRHI.cpp
+++ b/Engine/Plugins/Media/PixelCapture/Source/PixelCapture/Private/PixelCaptureCapturerRHI.cpp
@@ -67,35 +67,22 @@ void FPixelCaptureCapturerRHI::BeginProcess(const IPixelCaptureInputFrame& Input
 
 void FPixelCaptureCapturerRHI::CheckComplete()
 {
+	// TODO: We should move to proper event driven fences once they're implemented. Both DX12 and Vulkan APIs support them, they just haven't been added
+	// to their respective RHIs. DX11_3 supports it but for compatability reasons we can't upgrade from DX11_2.
+	
 	// in lieu of a proper callback we need to capture a thread to poll the fence
 	// so we know as quickly as possible when we can readback.
 
-	if (RHIType == ERHIInterfaceType::D3D11)
+	// sometimes we end up in a deadlock when we loop here polling the fence
+	// so instead we check and then submit a new check task.
+	if (!Fence->Poll())
 	{
-		// sometimes on dx11 we end up in a deadlock when we loop here polling the fence
-		// so instead we check and then submit a new check task.
-		if (!Fence->Poll())
-		{
-			TSharedRef<FPixelCaptureCapturerRHI> ThisRHIRef = StaticCastSharedRef<FPixelCaptureCapturerRHI>(AsShared());
-			AsyncTask(ENamedThreads::AnyBackgroundHiPriTask, [ThisRHIRef]() { ThisRHIRef->CheckComplete(); });
-		}
-		else
-		{
-			OnRHIStageComplete();
-		}
+		TSharedRef<FPixelCaptureCapturerRHI> ThisRHIRef = StaticCastSharedRef<FPixelCaptureCapturerRHI>(AsShared());
+		AsyncTask(ENamedThreads::AnyBackgroundHiPriTask, [ThisRHIRef]() { ThisRHIRef->CheckComplete(); });
 	}
 	else
 	{
-		// we submit a task here to loop poll the fence so that we get notified as quickly as possible about its completion.
-		TSharedRef<FPixelCaptureCapturerRHI> ThisRHIRef = StaticCastSharedRef<FPixelCaptureCapturerRHI>(AsShared());
-		AsyncTask(ENamedThreads::AnyBackgroundHiPriTask, [ThisRHIRef]() {
-			while (!ThisRHIRef->Fence->Poll())
-			{
-				// we want to check quickly but we dont want to burn CPU.
-				FPlatformProcess::Sleep(0.0001f);
-			}
-			ThisRHIRef->OnRHIStageComplete();
-		});
+		OnRHIStageComplete();
 	}
 }
 
