Add: PixelCapture on Mac

diff --git a/Engine/Plugins/Media/PixelCapture/PixelCapture.uplugin b/Engine/Plugins/Media/PixelCapture/PixelCapture.uplugin
--- a/Engine/Plugins/Media/PixelCapture/PixelCapture.uplugin
+++ b/Engine/Plugins/Media/PixelCapture/PixelCapture.uplugin
@@ -19,14 +19,14 @@
 			"Name": "PixelCapture",
 			"Type": "Runtime",
 			"LoadingPhase": "PreDefault",
-			"PlatformAllowList": [ "Win64", "Linux" ],
+			"PlatformAllowList": [ "Win64", "Linux", "Mac" ],
 			"TargetDenyList": [ "Server" ]
 		},
 		{
             "Name": "PixelCaptureShaders",
             "Type": "Runtime",
             "LoadingPhase": "PostConfigInit",
-            "PlatformAllowList": [ "Win64", "Linux" ],
+            "PlatformAllowList": [ "Win64", "Linux", "Mac" ],
             "TargetDenyList": [ "Server" ]
 		}
 	],
diff --git a/Engine/Plugins/Media/PixelCapture/Source/PixelCapture/Private/PixelCaptureCapturerRHI.cpp b/Engine/Plugins/Media/PixelCapture/Source/PixelCapture/Private/PixelCaptureCapturerRHI.cpp
--- a/Engine/Plugins/Media/PixelCapture/Source/PixelCapture/Private/PixelCaptureCapturerRHI.cpp
+++ b/Engine/Plugins/Media/PixelCapture/Source/PixelCapture/Private/PixelCaptureCapturerRHI.cpp
@@ -25,18 +25,28 @@ IPixelCaptureOutputFrame* FPixelCaptureCapturerRHI::CreateOutputBuffer(int32 Inp
 	const int32 Width = InputWidth * Scale;
 	const int32 Height = InputHeight * Scale;
 
-	FRHITextureCreateDesc TextureDesc =
-		FRHITextureCreateDesc::Create2D(TEXT("FFrameDataH264 Texture"), Width, Height, EPixelFormat::PF_B8G8R8A8)
-			.SetClearValue(FClearValueBinding::None)
+	FRHITextureCreateDesc TextureDesc = FRHITextureCreateDesc::Create2D(TEXT("FPixelCaptureCapturerRHI Texture"), Width, Height, EPixelFormat::PF_B8G8R8A8);
+
+	if (RHIGetInterfaceType() == ERHIInterfaceType::Metal)
+	{
+		TextureDesc.SetClearValue(FClearValueBinding::None)
+			.SetFlags(ETextureCreateFlags::CPUReadback)
+			.SetInitialState(ERHIAccess::CPURead);
+	}
+	else if (RHIGetInterfaceType() == ERHIInterfaceType::D3D11 || RHIGetInterfaceType() == ERHIInterfaceType::D3D12 || RHIGetInterfaceType() == ERHIInterfaceType::Vulkan)
+	{
+		TextureDesc.SetClearValue(FClearValueBinding::None)
 			.SetFlags(ETextureCreateFlags::RenderTargetable)
-			.SetInitialState(ERHIAccess::Present)
-			.DetermineInititialState();
+			.SetInitialState(ERHIAccess::Present);
+	}
+
+	TextureDesc.DetermineInititialState();		
 
 	if (RHIGetInterfaceType() == ERHIInterfaceType::Vulkan)
 	{
 		TextureDesc.AddFlags(ETextureCreateFlags::External);
 	}
-	else
+	else if (RHIGetInterfaceType() == ERHIInterfaceType::D3D11 || RHIGetInterfaceType() == ERHIInterfaceType::D3D12)
 	{
 		TextureDesc.AddFlags(ETextureCreateFlags::Shared);
 	}
diff --git a/Engine/Plugins/Media/PixelCapture/Source/PixelCapture/Private/PixelCaptureCapturerRHIRDG.cpp b/Engine/Plugins/Media/PixelCapture/Source/PixelCapture/Private/PixelCaptureCapturerRHIRDG.cpp
--- a/Engine/Plugins/Media/PixelCapture/Source/PixelCapture/Private/PixelCaptureCapturerRHIRDG.cpp
+++ b/Engine/Plugins/Media/PixelCapture/Source/PixelCapture/Private/PixelCaptureCapturerRHIRDG.cpp
@@ -23,18 +23,28 @@ IPixelCaptureOutputFrame* FPixelCaptureCapturerRHIRDG::CreateOutputBuffer(int32
 	const int32 Width = InputWidth * Scale;
 	const int32 Height = InputHeight * Scale;
 
-	FRHITextureCreateDesc TextureDesc =
-		FRHITextureCreateDesc::Create2D(TEXT("FFrameDataH264 Texture"), Width, Height, EPixelFormat::PF_B8G8R8A8)
-			.SetClearValue(FClearValueBinding::None)
+	FRHITextureCreateDesc TextureDesc = FRHITextureCreateDesc::Create2D(TEXT("FPixelCaptureCapturerRHIRDG Texture"), Width, Height, EPixelFormat::PF_B8G8R8A8);
+
+	if (RHIGetInterfaceType() == ERHIInterfaceType::Metal)
+	{
+		TextureDesc.SetClearValue(FClearValueBinding::None)
+			.SetFlags(ETextureCreateFlags::CPUReadback)
+			.SetInitialState(ERHIAccess::CPURead);
+	}
+	else if (RHIGetInterfaceType() == ERHIInterfaceType::D3D11 || RHIGetInterfaceType() == ERHIInterfaceType::D3D12 || RHIGetInterfaceType() == ERHIInterfaceType::Vulkan)
+	{
+		TextureDesc.SetClearValue(FClearValueBinding::None)
 			.SetFlags(ETextureCreateFlags::RenderTargetable)
-			.SetInitialState(ERHIAccess::Present)
-			.DetermineInititialState();
+			.SetInitialState(ERHIAccess::Present);
+	}
+
+	TextureDesc.DetermineInititialState();		
 
 	if (RHIGetInterfaceType() == ERHIInterfaceType::Vulkan)
 	{
 		TextureDesc.AddFlags(ETextureCreateFlags::External);
 	}
-	else
+	else if (RHIGetInterfaceType() == ERHIInterfaceType::D3D11 || RHIGetInterfaceType() == ERHIInterfaceType::D3D12)
 	{
 		TextureDesc.AddFlags(ETextureCreateFlags::Shared);
 	}
diff --git a/Engine/Plugins/Media/PixelCapture/Source/PixelCapture/Private/PixelCaptureCapturerRHIToI420CPU.cpp b/Engine/Plugins/Media/PixelCapture/Source/PixelCapture/Private/PixelCaptureCapturerRHIToI420CPU.cpp
--- a/Engine/Plugins/Media/PixelCapture/Source/PixelCapture/Private/PixelCaptureCapturerRHIToI420CPU.cpp
+++ b/Engine/Plugins/Media/PixelCapture/Source/PixelCapture/Private/PixelCaptureCapturerRHIToI420CPU.cpp
@@ -6,6 +6,9 @@
 #include "PixelCaptureBufferFormat.h"
 #include "PixelCaptureUtils.h"
 
+#include "Async/Async.h"
+#include "Containers/UnrealString.h"
+
 #include "libyuv/convert.h"
 
 TSharedPtr<FPixelCaptureCapturerRHIToI420CPU> FPixelCaptureCapturerRHIToI420CPU::Create(float InScale)
@@ -20,7 +23,6 @@ FPixelCaptureCapturerRHIToI420CPU::FPixelCaptureCapturerRHIToI420CPU(float InSca
 
 FPixelCaptureCapturerRHIToI420CPU::~FPixelCaptureCapturerRHIToI420CPU()
 {
-	CleanUp();
 }
 
 void FPixelCaptureCapturerRHIToI420CPU::Initialize(int32 InputWidth, int32 InputHeight)
@@ -46,18 +48,7 @@ void FPixelCaptureCapturerRHIToI420CPU::Initialize(int32 InputWidth, int32 Input
 
 	StagingTexture = RHICreateTexture(TextureDesc);
 
-	FRHITextureCreateDesc ReadbackDesc =
-		FRHITextureCreateDesc::Create2D(TEXT("FPixelCaptureCapturerRHIToI420CPU ReadbackTexture"), Width, Height, EPixelFormat::PF_B8G8R8A8)
-			.SetClearValue(FClearValueBinding::None)
-			.SetFlags(ETextureCreateFlags::CPUReadback)
-			.SetInitialState(ERHIAccess::CPURead)
-			.DetermineInititialState();
-
-	ReadbackTexture = RHICreateTexture(ReadbackDesc);
-
-	int32 BufferWidth = 0, BufferHeight = 0;
-	GDynamicRHI->RHIMapStagingSurface(ReadbackTexture, nullptr, ResultsBuffer, BufferWidth, BufferHeight);
-	MappedStride = BufferWidth;
+	TextureReader = MakeShared<FRHIGPUTextureReadback>(TEXT("FPixelCaptureCapturerRHIToI420CPUReadback"));
 
 	FPixelCaptureCapturer::Initialize(InputWidth, InputHeight);
 }
@@ -85,12 +76,9 @@ void FPixelCaptureCapturerRHIToI420CPU::BeginProcess(const IPixelCaptureInputFra
 	RHICmdList.Transition(FRHITransitionInfo(SourceTexture, ERHIAccess::Unknown, ERHIAccess::CopySrc));
 	RHICmdList.Transition(FRHITransitionInfo(StagingTexture, ERHIAccess::CopySrc, ERHIAccess::CopyDest));
 	CopyTexture(RHICmdList, SourceTexture, StagingTexture, nullptr);
-
 	RHICmdList.Transition(FRHITransitionInfo(StagingTexture, ERHIAccess::CopyDest, ERHIAccess::CopySrc));
-	RHICmdList.Transition(FRHITransitionInfo(ReadbackTexture, ERHIAccess::CPURead, ERHIAccess::CopyDest));
-	RHICmdList.CopyTexture(StagingTexture, ReadbackTexture, {});
 
-	RHICmdList.Transition(FRHITransitionInfo(ReadbackTexture, ERHIAccess::CopyDest, ERHIAccess::CPURead));
+	TextureReader->EnqueueCopy(RHICmdList, StagingTexture, FIntVector(0, 0, 0), 0, FIntVector(StagingTexture->GetSizeXY().X, StagingTexture->GetSizeXY().Y, 0));
 
 	MarkCPUWorkEnd();
 
@@ -98,20 +86,37 @@ void FPixelCaptureCapturerRHIToI420CPU::BeginProcess(const IPixelCaptureInputFra
 	// until after the rhi thread is done with it, so all the commands will still have valid references.
 	TSharedRef<FPixelCaptureCapturerRHIToI420CPU> ThisRHIRef = StaticCastSharedRef<FPixelCaptureCapturerRHIToI420CPU>(AsShared());
 	RHICmdList.EnqueueLambda([ThisRHIRef, OutputBuffer](FRHICommandListImmediate&) {
-		ThisRHIRef->OnRHIStageComplete(OutputBuffer);
+		ThisRHIRef->CheckComplete(OutputBuffer);
 	});
 }
 
+void FPixelCaptureCapturerRHIToI420CPU::CheckComplete(IPixelCaptureOutputFrame* OutputBuffer)
+{
+	if (!TextureReader->IsReady())
+	{
+		TSharedRef<FPixelCaptureCapturerRHIToI420CPU> ThisRHIRef = StaticCastSharedRef<FPixelCaptureCapturerRHIToI420CPU>(AsShared());
+		AsyncTask(ENamedThreads::AnyBackgroundHiPriTask, [ThisRHIRef, OutputBuffer]() { ThisRHIRef->CheckComplete(OutputBuffer); });
+	}
+	else
+	{
+		TSharedRef<FPixelCaptureCapturerRHIToI420CPU> ThisRHIRef = StaticCastSharedRef<FPixelCaptureCapturerRHIToI420CPU>(AsShared());
+		AsyncTask(ENamedThreads::ActualRenderingThread, [ThisRHIRef, OutputBuffer]() { ThisRHIRef->OnRHIStageComplete(OutputBuffer); });
+	}
+}
+
 void FPixelCaptureCapturerRHIToI420CPU::OnRHIStageComplete(IPixelCaptureOutputFrame* OutputBuffer)
 {
 	MarkGPUWorkEnd();
 	MarkCPUWorkStart();
 
+	int32 OutRowPitchInPixels = 0;
+	void* PixelData = TextureReader->Lock(OutRowPitchInPixels);
+
 	FPixelCaptureOutputFrameI420* OutputI420Buffer = StaticCast<FPixelCaptureOutputFrameI420*>(OutputBuffer);
 	TSharedPtr<FPixelCaptureBufferI420> I420Buffer = OutputI420Buffer->GetI420Buffer();
 	libyuv::ARGBToI420(
-		static_cast<uint8*>(ResultsBuffer),
-		MappedStride * 4,
+		static_cast<uint8*>(PixelData),
+		OutRowPitchInPixels * 4,
 		I420Buffer->GetMutableDataY(),
 		I420Buffer->GetStrideY(),
 		I420Buffer->GetMutableDataU(),
@@ -121,12 +126,8 @@ void FPixelCaptureCapturerRHIToI420CPU::OnRHIStageComplete(IPixelCaptureOutputFr
 		I420Buffer->GetWidth(),
 		I420Buffer->GetHeight());
 
+	TextureReader->Unlock();
+
 	MarkCPUWorkEnd();
 	EndProcess();
 }
-
-void FPixelCaptureCapturerRHIToI420CPU::CleanUp()
-{
-	GDynamicRHI->RHIUnmapStagingSurface(ReadbackTexture);
-	ResultsBuffer = nullptr;
-}
diff --git a/Engine/Plugins/Media/PixelCapture/Source/PixelCapture/Public/PixelCaptureCapturerRHIToI420CPU.h b/Engine/Plugins/Media/PixelCapture/Source/PixelCapture/Public/PixelCaptureCapturerRHIToI420CPU.h
--- a/Engine/Plugins/Media/PixelCapture/Source/PixelCapture/Public/PixelCaptureCapturerRHIToI420CPU.h
+++ b/Engine/Plugins/Media/PixelCapture/Source/PixelCapture/Public/PixelCaptureCapturerRHIToI420CPU.h
@@ -4,6 +4,7 @@
 
 #include "PixelCaptureCapturer.h"
 #include "RHI.h"
+#include "RHIGPUReadback.h"
 
 /**
  * A basic capturer that will capture RHI texture frames to I420 buffers utilizing cpu functions.
@@ -30,12 +31,10 @@ protected:
 private:
 	float Scale = 1.0f;
 
+	TSharedPtr<FRHIGPUTextureReadback> TextureReader;
 	FTextureRHIRef StagingTexture;
-	FTextureRHIRef ReadbackTexture;
-	void* ResultsBuffer = nullptr;
-	int32 MappedStride = 0;
 
 	FPixelCaptureCapturerRHIToI420CPU(float InScale);
 	void OnRHIStageComplete(IPixelCaptureOutputFrame* OutputBuffer);
-	void CleanUp();
+	void CheckComplete(IPixelCaptureOutputFrame* OutputBuffer);
 };
