Fix: Streaming level editor on Mac would hard crash the editor

diff --git a/Engine/Plugins/Media/PixelCapture/Source/PixelCapture/Private/PixelCaptureUtils.h b/Engine/Plugins/Media/PixelCapture/Source/PixelCapture/Private/PixelCaptureUtils.h
--- a/Engine/Plugins/Media/PixelCapture/Source/PixelCapture/Private/PixelCaptureUtils.h
+++ b/Engine/Plugins/Media/PixelCapture/Source/PixelCapture/Private/PixelCaptureUtils.h
@@ -115,6 +115,12 @@ inline void CopyTextureRDG(FRHICommandListImmediate& RHICmdList, FTextureRHIRef
 		// The todo is here because there has to be a better way to achieve what we want without
 		// all of this song and dance.
 
+#if PLATFORM_MAC
+		// Create a staging texture that is the same size and format as the final.
+		FRDGTextureRef StagingTexture = GraphBuilder.CreateTexture(FRDGTextureDesc::Create2D(FIntPoint(OutputTexture->Desc.Extent.X, OutputTexture->Desc.Extent.Y), OutputTexture->Desc.Format, OutputTexture->Desc.ClearValue, ETextureCreateFlags::RenderTargetable), TEXT("PixelCaptureCopy Staging"));
+		FScreenPassTextureViewport StagingViewport(StagingTexture);
+#endif
+
 		// The formats or size differ to pixel shader stuff
 		//Configure source/output viewport to get the right UV scaling from source texture to output texture
 		FScreenPassTextureViewport InputViewport(InputTexture);
@@ -132,7 +138,15 @@ inline void CopyTextureRDG(FRHICommandListImmediate& RHICmdList, FTextureRHIRef
 		const FIntRect ViewRect(FIntPoint(0, 0), InputTexture->Desc.Extent);
 
 		TShaderMapRef<FModifyAlphaSwizzleRgbaPS> PixelShader(GlobalShaderMap, PermutationVector);
-		FModifyAlphaSwizzleRgbaPS::FParameters* PixelShaderParameters = PixelShader->AllocateAndSetParameters(GraphBuilder, InputTexture, OutputTexture);
+		FModifyAlphaSwizzleRgbaPS::FParameters* PixelShaderParameters = PixelShader->AllocateAndSetParameters(
+			GraphBuilder,
+			InputTexture,
+#if PLATFORM_MAC
+			StagingTexture
+#else
+			OutputTexture
+#endif
+			);
 		
 		FRHIBlendState* BlendState = FScreenPassPipelineState::FDefaultBlendState::GetRHI();
 		FRHIDepthStencilState* DepthStencilState = FScreenPassPipelineState::FDefaultDepthStencilState::GetRHI();
@@ -141,11 +155,28 @@ inline void CopyTextureRDG(FRHICommandListImmediate& RHICmdList, FTextureRHIRef
 			GraphBuilder,
 			RDG_EVENT_NAME("PixelCapturerSwizzle"),
 			FScreenPassViewInfo(),
+#if PLATFORM_MAC
+			StagingViewport,
+#else
 			OutputViewport,
+#endif
+			
 			InputViewport,
 			VertexShader,
 			PixelShader,
 			PixelShaderParameters);
+
+#if PLATFORM_MAC
+		// Now we can be certain the formats are the same and size are the same. simple copy
+		AddDrawTexturePass(
+			GraphBuilder,
+			GetGlobalShaderMap(GMaxRHIFeatureLevel),
+			StagingTexture,
+			OutputTexture,
+			FRDGDrawTextureInfo()
+		);
+#endif
 	}
+
 	GraphBuilder.Execute();
 }
