Add: I040

NOTE: No capture pipeline exists for I010 data. This buffer exists purely to copy converted I4xx data in to.

diff --git a/Engine/Plugins/Media/PixelCapture/Source/PixelCapture/Private/PixelCaptureBufferI010.cpp b/Engine/Plugins/Media/PixelCapture/Source/PixelCapture/Private/PixelCaptureBufferI010.cpp
new file mode 100644
--- /dev/null
+++ b/Engine/Plugins/Media/PixelCapture/Source/PixelCapture/Private/PixelCaptureBufferI010.cpp
@@ -0,0 +1,70 @@
+// Copyright Epic Games, Inc. All Rights Reserved.
+
+#include "PixelCaptureBufferI010.h"
+
+namespace
+{
+	int CalcI010Size(int StrideY, int StrideUV, int Height)
+	{
+		return 2 * (StrideY * Height + (StrideUV + StrideUV) * ((Height + 1) / 2));
+	}
+} // namespace
+
+FPixelCaptureBufferI010::FPixelCaptureBufferI010(int InWidth, int InHeight)
+	: Width(InWidth)
+	, Height(InHeight)
+	, StrideY(InWidth)
+	, StrideUV((InWidth + 1) / 2)
+{
+	Data.SetNum(CalcI010Size(StrideY, StrideUV, Height));
+}
+
+const uint8_t* FPixelCaptureBufferI010::GetData() const
+{
+	return reinterpret_cast<const uint8_t*>(Data.GetData());
+}
+
+const uint16_t* FPixelCaptureBufferI010::GetDataY() const
+{
+	return Data.GetData();
+}
+
+const uint16_t* FPixelCaptureBufferI010::GetDataU() const
+{
+	return GetDataY() + GetDataSizeY();
+}
+
+const uint16_t* FPixelCaptureBufferI010::GetDataV() const
+{
+	return GetDataU() + GetDataSizeUV();
+}
+
+uint8_t* FPixelCaptureBufferI010::GetMutableData()
+{
+	return reinterpret_cast<uint8_t*>(Data.GetData());
+}
+
+uint16_t* FPixelCaptureBufferI010::GetMutableDataY()
+{
+	return Data.GetData();
+}
+
+uint16_t* FPixelCaptureBufferI010::GetMutableDataU()
+{
+	return GetMutableDataY() + GetDataSizeY();
+}
+
+uint16_t* FPixelCaptureBufferI010::GetMutableDataV()
+{
+	return GetMutableDataU() + GetDataSizeUV();
+}
+
+int FPixelCaptureBufferI010::GetDataSizeY() const
+{
+	return StrideY * Height;
+}
+
+int FPixelCaptureBufferI010::GetDataSizeUV() const
+{
+	return StrideUV * ((Height + 1) / 2);
+}
diff --git a/Engine/Plugins/Media/PixelCapture/Source/PixelCapture/Public/PixelCaptureBufferFormat.h b/Engine/Plugins/Media/PixelCapture/Source/PixelCapture/Public/PixelCaptureBufferFormat.h
--- a/Engine/Plugins/Media/PixelCapture/Source/PixelCapture/Public/PixelCaptureBufferFormat.h
+++ b/Engine/Plugins/Media/PixelCapture/Source/PixelCapture/Public/PixelCaptureBufferFormat.h
@@ -10,11 +10,12 @@
  */
 namespace PixelCaptureBufferFormat
 {
-    static const int32 FORMAT_UNKNOWN = 0;
-    static const int32 FORMAT_RHI = 1;
-    static const int32 FORMAT_I420 = 2;
-    static const int32 FORMAT_NV12 = 3;
-    static const int32 FORMAT_P010 = 4;
+	static const int32 FORMAT_UNKNOWN = 0;
+	static const int32 FORMAT_RHI = 1;
+	static const int32 FORMAT_I420 = 2;
+	static const int32 FORMAT_NV12 = 3;
+	static const int32 FORMAT_P010 = 4;
+	static const int32 FORMAT_I010 = 5;
 
-    static const int32 FORMAT_USER = 128;
-}
+	static const int32 FORMAT_USER = 128;
+} // namespace PixelCaptureBufferFormat
diff --git a/Engine/Plugins/Media/PixelCapture/Source/PixelCapture/Public/PixelCaptureBufferI010.h b/Engine/Plugins/Media/PixelCapture/Source/PixelCapture/Public/PixelCaptureBufferI010.h
new file mode 100644
--- /dev/null
+++ b/Engine/Plugins/Media/PixelCapture/Source/PixelCapture/Public/PixelCaptureBufferI010.h
@@ -0,0 +1,115 @@
+// Copyright Epic Games, Inc. All Rights Reserved.
+
+#pragma once
+
+#include "IPixelCaptureBuffer.h"
+
+/**
+ * Buffer class for holding I010 frame data.
+ */
+class PIXELCAPTURE_API FPixelCaptureBufferI010 : public IPixelCaptureBuffer
+{
+public:
+	/**
+	 * Create an empty buffer of the specified size.
+	 */
+	FPixelCaptureBufferI010(int InWidth, int InHeight);
+
+	FPixelCaptureBufferI010(const FPixelCaptureBufferI010& Other) = default;
+	virtual ~FPixelCaptureBufferI010() = default;
+
+	/**
+	 * Get the width of the frame
+	 */
+	virtual int32 GetWidth() const override { return Width; }
+
+	/**
+	 * Get the height of the frame
+	 */
+	virtual int32 GetHeight() const override { return Height; }
+
+	/**
+	 * Get the Buffer format
+	 */
+	virtual int32 GetFormat() const override { return PixelCaptureBufferFormat::FORMAT_I010; };
+
+	/**
+	 * Get the stride of the Y plane
+	 */
+	int32 GetStrideY() const { return StrideY; }
+
+	/**
+	 * Get the stride of the U and V planes
+	 */
+	int32 GetStrideUV() const { return StrideUV; }
+
+	/**
+	 * Gets a const pointer to the beginning of entire buffer for reading.
+	 */
+	virtual const uint8_t* GetData() const override;
+
+	/**
+	 * Gets a const pointer to the beginning of the Y plane for reading.
+	 */
+	const uint16_t* GetDataY() const;
+
+	/**
+	 * Gets a const pointer to the beginning of the U plane for reading.
+	 */
+	const uint16_t* GetDataU() const;
+
+	/**
+	 * Gets a const pointer to the beginning of the V plane for reading.
+	 */
+	const uint16_t* GetDataV() const;
+
+	/**
+	 * Gets a pointer to the beginning of the entire buffer for editing.
+	 */
+	virtual uint8_t* GetMutableData() override;
+
+	/**
+	 * Gets a pointer to the beginning of the Y plane for editing.
+	 */
+	uint16_t* GetMutableDataY();
+
+	/**
+	 * Gets a pointer to the beginning of the U plane for editing.
+	 */
+	uint16_t* GetMutableDataU();
+
+	/**
+	 * Gets a pointer to the beginning of the V plane for editing.
+	 */
+	uint16_t* GetMutableDataV();
+
+	/**
+	 * Gets the size of the Y plane in bytes.
+	 */
+	int32 GetDataSizeY() const;
+
+	/**
+	 * Gets the size of both the U and V planes in bytes.
+	 */
+	int32 GetDataSizeUV() const;
+
+	/**
+	 * Gets the total size of this buffer.
+	 */
+	virtual int64 GetSize() const override { return Data.Num(); }
+
+	/**
+	 * Copies the given buffer into this buffer. Supplied so copies
+	 * can be explicit in code where we want to be verbose.
+	 */
+	void Copy(const FPixelCaptureBufferI010& Other) { *this = Other; }
+
+private:
+	int Width = 0;
+	int Height = 0;
+
+	int StrideY = 0;
+	int StrideUV = 0;
+
+	TArray<uint16_t> Data;
+};
