QOL: Update number of maximum capturers to support a large number of concurrent streams

diff --git a/Engine/Plugins/Media/PixelCapture/Source/PixelCapture/Private/PixelCaptureCapturer.cpp b/Engine/Plugins/Media/PixelCapture/Source/PixelCapture/Private/PixelCaptureCapturer.cpp
--- a/Engine/Plugins/Media/PixelCapture/Source/PixelCapture/Private/PixelCaptureCapturer.cpp
+++ b/Engine/Plugins/Media/PixelCapture/Source/PixelCapture/Private/PixelCaptureCapturer.cpp
@@ -56,7 +56,7 @@ void FPixelCaptureCapturer::Initialize(int32 InputWidth, int32 InputHeight)
 {
 	checkf(InputWidth > 0 && InputHeight > 0, TEXT("Capture should be initialized with non-zero resolution."));
 	Buffer = MakeUnique<UE::PixelCapture::FOutputFrameBuffer>();
-	Buffer->Reset(3, 10, [this, InputWidth, InputHeight]() { return TSharedPtr<IPixelCaptureOutputFrame>(CreateOutputBuffer(InputWidth, InputHeight)); });
+	Buffer->Reset(3, 32, [this, InputWidth, InputHeight]() { return TSharedPtr<IPixelCaptureOutputFrame>(CreateOutputBuffer(InputWidth, InputHeight)); });
 	ExpectedInputWidth = InputWidth;
 	ExpectedInputHeight = InputHeight;
 	bHasOutput = false;
