[PixelCapture] Add: New MediaIO based PixelCapturer that handles RHI -> RHI and RHI -> I420 CPU conversion

diff --git a/Engine/Plugins/Media/PixelCapture/PixelCapture.uplugin b/Engine/Plugins/Media/PixelCapture/PixelCapture.uplugin
--- a/Engine/Plugins/Media/PixelCapture/PixelCapture.uplugin
+++ b/Engine/Plugins/Media/PixelCapture/PixelCapture.uplugin
@@ -30,5 +30,10 @@
             "TargetDenyList": [ "Server" ]
 		}
 	],
-	"Plugins": []
+	"Plugins": [
+		{
+			"Name": "MediaIOFramework",
+			"Enabled": true
+		}
+	]
 }
diff --git a/Engine/Plugins/Media/PixelCapture/Source/PixelCapture/Internal/PixelCaptureCapturerMediaCapture.h b/Engine/Plugins/Media/PixelCapture/Source/PixelCapture/Internal/PixelCaptureCapturerMediaCapture.h
new file mode 100644
--- /dev/null
+++ b/Engine/Plugins/Media/PixelCapture/Source/PixelCapture/Internal/PixelCaptureCapturerMediaCapture.h
@@ -0,0 +1,126 @@
+// Copyright Epic Games, Inc. All Rights Reserved.
+
+#pragma once
+
+#include "MediaCapture.h"
+#include "MediaOutput.h"
+#include "PixelCaptureBufferFormat.h"
+#include "PixelCaptureCapturer.h"
+
+#include "PixelCaptureCapturerMediaCapture.generated.h"
+
+UCLASS(BlueprintType)
+class PIXELCAPTURE_API UPixelCaptureMediaOuput : public UMediaOutput
+{
+	GENERATED_BODY()
+
+public:
+	void									 SetRequestedSize(FIntPoint InRequestedSize) { RequestedSize = InRequestedSize; }
+	virtual FIntPoint						 GetRequestedSize() const override { return RequestedSize; }
+	virtual EPixelFormat					 GetRequestedPixelFormat() const override { return EPixelFormat::PF_B8G8R8A8; }
+	virtual EMediaCaptureConversionOperation GetConversionOperation(EMediaCaptureSourceType InSourceType) const override { return EMediaCaptureConversionOperation::CUSTOM; }
+
+private:
+	FIntPoint RequestedSize;
+};
+
+UCLASS(BlueprintType)
+class PIXELCAPTURE_API UPixelCaptureMediaCapture : public UMediaCapture
+{
+	GENERATED_BODY()
+
+public:
+	void SetOutputFrame(IPixelCaptureOutputFrame* InOutputFrame)
+	{
+		OutputFrame = InOutputFrame;
+	}
+
+	void SetFormat(int32 InFormat)
+	{
+		Format = InFormat;
+	}
+
+	DECLARE_EVENT(UPixelCaptureMediaCapture, FOnCaptureComplete)
+	FOnCaptureComplete OnCaptureComplete;
+
+protected:
+	virtual bool InitializeCapture() override
+	{
+		SetState(EMediaCaptureState::Capturing);
+		return true;
+	}
+
+	virtual void OnRHIResourceCaptured_AnyThread(
+		const FCaptureBaseData&								   BaseData,
+		TSharedPtr<FMediaCaptureUserData, ESPMode::ThreadSafe> UserData,
+		FTextureRHIRef										   Texture) override;
+
+	virtual void OnFrameCaptured_AnyThread(
+		const FCaptureBaseData&								   BaseData,
+		TSharedPtr<FMediaCaptureUserData, ESPMode::ThreadSafe> UserData,
+		const FMediaCaptureResourceData&					   ResourceData) override;
+
+	virtual void OnCustomCapture_RenderingThread(
+		FRDGBuilder&										   GraphBuilder,
+		const FCaptureBaseData&								   InBaseData,
+		TSharedPtr<FMediaCaptureUserData, ESPMode::ThreadSafe> InUserData,
+		FRDGTextureRef										   InSourceTexture,
+		FRDGTextureRef										   OutputTexture,
+		const FRHICopyTextureInfo&							   CopyInfo,
+		FVector2D											   CropU,
+		FVector2D											   CropV) override;
+
+	virtual bool ShouldCaptureRHIResource() const override
+	{
+		return Format == PixelCaptureBufferFormat::FORMAT_RHI;
+	}
+
+	virtual bool SupportsAnyThreadCapture() const
+	{
+		return true;
+	}
+
+	virtual ETextureCreateFlags GetOutputTextureFlags() const override;
+
+private:
+	IPixelCaptureOutputFrame* OutputFrame = nullptr;
+	int32					  Format = PixelCaptureBufferFormat::FORMAT_UNKNOWN;
+};
+
+/**
+ * A MediaIO based capturer that will copy and convert RHI texture frames.
+ * Input: FPixelCaptureInputFrameRHI
+ * Output: FPixelCaptureOutputFrameRHI/FPixelCaptureOutputFrameI420
+ */
+class PIXELCAPTURE_API FPixelCaptureCapturerMediaCapture : public FPixelCaptureCapturer, public TSharedFromThis<FPixelCaptureCapturerMediaCapture>
+{
+public:
+	/**
+	 * Creates a new Capturer capturing the input frame at the given scale.
+	 * @param InScale The scale of the resulting output capture.
+	 */
+	static TSharedPtr<FPixelCaptureCapturerMediaCapture> Create(float InScale, int32 InFormat);
+	virtual ~FPixelCaptureCapturerMediaCapture() override;
+
+protected:
+	virtual FString					  GetCapturerName() const override { return "FPixelCaptureCapturerMediaCapture Copy"; }
+	virtual IPixelCaptureOutputFrame* CreateOutputBuffer(int32 InputWidth, int32 InputHeight) override;
+	virtual void					  BeginProcess(const IPixelCaptureInputFrame& InputFrame, IPixelCaptureOutputFrame* OutputBuffer) override;
+
+private:
+	float Scale = 1.0f;
+	int32 Format = PixelCaptureBufferFormat::FORMAT_UNKNOWN;
+	bool  bMediaCaptureInitialized = false;
+
+	FPixelCaptureCapturerMediaCapture(float InScale, int32 InFormat);
+	void InitializeMediaCapture();
+	void LateStartCapturer();
+
+	UPROPERTY(Transient)
+	TObjectPtr<UPixelCaptureMediaCapture> MediaCapture = nullptr;
+
+	UPROPERTY(Transient)
+	TObjectPtr<UPixelCaptureMediaOuput> MediaOutput = nullptr;
+
+	TOptional<FDelegateHandle> OnFrameEndDelegateHandle;
+};
diff --git a/Engine/Plugins/Media/PixelCapture/Source/PixelCapture/PixelCapture.build.cs b/Engine/Plugins/Media/PixelCapture/Source/PixelCapture/PixelCapture.build.cs
--- a/Engine/Plugins/Media/PixelCapture/Source/PixelCapture/PixelCapture.build.cs
+++ b/Engine/Plugins/Media/PixelCapture/Source/PixelCapture/PixelCapture.build.cs
@@ -16,10 +16,12 @@ namespace UnrealBuildTool.Rules
 			// This is so for game projects using our public headers don't have to include extra modules they might not know about.
 			PublicDependencyModuleNames.AddRange(new string[] {
 				"RHI",
+				"MediaIOCore",
 			});
 
 			PrivateDependencyModuleNames.AddRange(new string[] {
 				"Core",
+				"CoreUObject",
 				"Engine",
 				"PixelCaptureShaders",
 				"RenderCore",
diff --git a/Engine/Plugins/Media/PixelCapture/Source/PixelCapture/Private/PixelCaptureCapturerLayered.cpp b/Engine/Plugins/Media/PixelCapture/Source/PixelCapture/Private/PixelCaptureCapturerLayered.cpp
--- a/Engine/Plugins/Media/PixelCapture/Source/PixelCapture/Private/PixelCaptureCapturerLayered.cpp
+++ b/Engine/Plugins/Media/PixelCapture/Source/PixelCapture/Private/PixelCaptureCapturerLayered.cpp
@@ -20,10 +20,12 @@ FPixelCaptureCapturerLayered::FPixelCaptureCapturerLayered(IPixelCaptureCapturer
 TSharedPtr<IPixelCaptureOutputFrame> FPixelCaptureCapturerLayered::ReadOutput(int32 LayerIndex)
 {
 	FScopeLock LayersLock(&LayersGuard);
-	if (!LayerCapturers.IsEmpty())
+	if (!LayerCapturers.IsEmpty() && LayerCapturers.Num() > LayerIndex)
 	{
 		return LayerCapturers[LayerIndex]->ReadOutput();
 	}
+
+	UE_LOG(LogPixelCapture, Warning, TEXT("No capturer configured for LayerIndex (%d)!"), LayerIndex);
 	return nullptr;
 }
 
diff --git a/Engine/Plugins/Media/PixelCapture/Source/PixelCapture/Private/PixelCaptureCapturerMediaCapture.cpp b/Engine/Plugins/Media/PixelCapture/Source/PixelCapture/Private/PixelCaptureCapturerMediaCapture.cpp
new file mode 100644
--- /dev/null
+++ b/Engine/Plugins/Media/PixelCapture/Source/PixelCapture/Private/PixelCaptureCapturerMediaCapture.cpp
@@ -0,0 +1,302 @@
+// Copyright Epic Games, Inc. All Rights Reserved.
+
+#include "PixelCaptureCapturerMediaCapture.h"
+#include "PixelCaptureBufferFormat.h"
+#include "PixelCaptureInputFrameRHI.h"
+#include "PixelCaptureOutputFrameRHI.h"
+#include "PixelCaptureOutputFrameI420.h"
+#include "PixelCaptureUtils.h"
+#include "PixelCapturePrivate.h"
+
+#include "Async/Async.h"
+
+#include "libyuv/convert.h"
+
+void UPixelCaptureMediaCapture::OnRHIResourceCaptured_AnyThread(
+	const FCaptureBaseData&								   BaseData,
+	TSharedPtr<FMediaCaptureUserData, ESPMode::ThreadSafe> UserData,
+	FTextureRHIRef										   Texture)
+{
+	if (OutputFrame == nullptr)
+	{
+		UE_LOGFMT(LogPixelCapture, Warning, "UPixelCaptureMediaCapture::OnRHIResourceCaptured_AnyThread: No output frame set!");
+		return;
+	}
+
+	check(Format == PixelCaptureBufferFormat::FORMAT_RHI);
+
+	static_cast<FPixelCaptureOutputFrameRHI*>(OutputFrame)->SetFrameTexture(Texture);
+
+	OnCaptureComplete.Broadcast();
+}
+
+void UPixelCaptureMediaCapture::OnFrameCaptured_AnyThread(
+	const FCaptureBaseData&								   BaseData,
+	TSharedPtr<FMediaCaptureUserData, ESPMode::ThreadSafe> UserData,
+	const FMediaCaptureResourceData&					   ResourceData)
+{
+	if (OutputFrame == nullptr)
+	{
+		UE_LOGFMT(LogPixelCapture, Warning, "UPixelCaptureMediaCapture::OnFrameCaptured_AnyThread: No output frame set!");
+		return;
+	}
+
+	check(Format == PixelCaptureBufferFormat::FORMAT_I420);
+
+	FPixelCaptureOutputFrameI420*		I420Frame = static_cast<FPixelCaptureOutputFrameI420*>(OutputFrame);
+	TSharedPtr<FPixelCaptureBufferI420> I420Buffer = MakeShared<FPixelCaptureBufferI420>(ResourceData.Width, ResourceData.Height);
+	libyuv::ARGBToI420(
+		static_cast<uint8*>(ResourceData.Buffer),
+		ResourceData.BytesPerRow,
+		I420Buffer->GetMutableDataY(),
+		I420Buffer->GetStrideY(),
+		I420Buffer->GetMutableDataU(),
+		I420Buffer->GetStrideUV(),
+		I420Buffer->GetMutableDataV(),
+		I420Buffer->GetStrideUV(),
+		I420Buffer->GetWidth(),
+		I420Buffer->GetHeight());
+
+	I420Frame->SetI420Buffer(I420Buffer);
+
+	OnCaptureComplete.Broadcast();
+}
+
+void UPixelCaptureMediaCapture::OnCustomCapture_RenderingThread(
+	FRDGBuilder&										   GraphBuilder,
+	const FCaptureBaseData&								   InBaseData,
+	TSharedPtr<FMediaCaptureUserData, ESPMode::ThreadSafe> InUserData,
+	FRDGTextureRef										   InSourceTexture,
+	FRDGTextureRef										   OutputTexture,
+	const FRHICopyTextureInfo&							   CopyInfo,
+	FVector2D											   CropU,
+	FVector2D											   CropV)
+{
+	bool bRequiresFormatConversion = InSourceTexture->Desc.Format != OutputTexture->Desc.Format;
+	if (InSourceTexture->Desc.Format == OutputTexture->Desc.Format
+		&& InSourceTexture->Desc.Extent.X == OutputTexture->Desc.Extent.X
+		&& InSourceTexture->Desc.Extent.Y == OutputTexture->Desc.Extent.Y)
+	{
+		// The formats are the same and size are the same. simple copy
+		AddDrawTexturePass(
+			GraphBuilder,
+			GetGlobalShaderMap(GMaxRHIFeatureLevel),
+			InSourceTexture,
+			OutputTexture,
+			FRDGDrawTextureInfo());
+
+		return;
+	}
+	else
+	{
+#if PLATFORM_MAC
+		// Create a staging texture that is the same size and format as the final.
+		FRDGTextureRef			   StagingTexture = GraphBuilder.CreateTexture(FRDGTextureDesc::Create2D(FIntPoint(OutputTexture->Desc.Extent.X, OutputTexture->Desc.Extent.Y), OutputTexture->Desc.Format, OutputTexture->Desc.ClearValue, ETextureCreateFlags::RenderTargetable), TEXT("PixelStreamingMediaIOCapture Staging"));
+		FScreenPassTextureViewport StagingViewport(StagingTexture);
+#endif
+
+		FScreenPassTextureViewport InputViewport(InSourceTexture);
+		FScreenPassTextureViewport OutputViewport(OutputTexture);
+
+		FGlobalShaderMap*			 GlobalShaderMap = GetGlobalShaderMap(GMaxRHIFeatureLevel);
+		TShaderMapRef<FScreenPassVS> VertexShader(GlobalShaderMap);
+
+		// In cases where texture is converted from a format that doesn't have A channel, we want to force set it to 1.
+		int32										  MediaConversionOperation = 0; // None
+		FModifyAlphaSwizzleRgbaPS::FPermutationDomain PermutationVector;
+		PermutationVector.Set<FModifyAlphaSwizzleRgbaPS::FConversionOp>(MediaConversionOperation);
+
+		// Rectangle area to use from source
+		const FIntRect ViewRect(FIntPoint(0, 0), InSourceTexture->Desc.Extent);
+
+		TShaderMapRef<FModifyAlphaSwizzleRgbaPS> PixelShader(GlobalShaderMap, PermutationVector);
+		FModifyAlphaSwizzleRgbaPS::FParameters*	 PixelShaderParameters = PixelShader->AllocateAndSetParameters(
+			 GraphBuilder,
+			 InSourceTexture,
+#if PLATFORM_MAC
+			StagingTexture
+#else
+			OutputTexture
+#endif
+		);
+
+		FRHIBlendState*		   BlendState = FScreenPassPipelineState::FDefaultBlendState::GetRHI();
+		FRHIDepthStencilState* DepthStencilState = FScreenPassPipelineState::FDefaultDepthStencilState::GetRHI();
+
+		AddDrawScreenPass(
+			GraphBuilder,
+			RDG_EVENT_NAME("PixelStreamingEpicRtcMediaIOCapture Swizzle"),
+			FScreenPassViewInfo(),
+#if PLATFORM_MAC
+			StagingViewport,
+#else
+			OutputViewport,
+#endif
+
+			InputViewport,
+			VertexShader,
+			PixelShader,
+			PixelShaderParameters);
+
+#if PLATFORM_MAC
+		// Now we can be certain the formats are the same and size are the same. simple copy
+		AddDrawTexturePass(
+			GraphBuilder,
+			GetGlobalShaderMap(GMaxRHIFeatureLevel),
+			StagingTexture,
+			OutputTexture,
+			FRDGDrawTextureInfo());
+#endif
+	}
+}
+
+ETextureCreateFlags UPixelCaptureMediaCapture::GetOutputTextureFlags() const
+{
+#if PLATFORM_MAC
+	return TexCreate_CPUReadback;
+#else
+	ETextureCreateFlags Flags = TexCreate_RenderTargetable | TexCreate_UAV;
+
+	if (RHIGetInterfaceType() == ERHIInterfaceType::Vulkan)
+	{
+		Flags |= TexCreate_External;
+	}
+	else if (RHIGetInterfaceType() == ERHIInterfaceType::D3D11 || RHIGetInterfaceType() == ERHIInterfaceType::D3D12)
+	{
+		Flags |= TexCreate_Shared;
+	}
+
+	return Flags;
+#endif
+}
+
+TSharedPtr<FPixelCaptureCapturerMediaCapture> FPixelCaptureCapturerMediaCapture::Create(float InScale, int32 InFormat)
+{
+	TSharedPtr<FPixelCaptureCapturerMediaCapture> Capturer = TSharedPtr<FPixelCaptureCapturerMediaCapture>(new FPixelCaptureCapturerMediaCapture(InScale, InFormat));
+	Capturer->LateStartCapturer();
+
+	return Capturer;
+}
+
+FPixelCaptureCapturerMediaCapture::FPixelCaptureCapturerMediaCapture(float InScale, int32 InFormat)
+	: Scale(InScale)
+	, Format(InFormat)
+{
+	if (Format != PixelCaptureBufferFormat::FORMAT_RHI && Format != PixelCaptureBufferFormat::FORMAT_I420)
+	{
+		UE_LOGFMT(LogPixelCapture, Warning, "FPixelCaptureCapturerMediaCapture: Invalid pixel format. Expected either FORMAT_RHI or FORMAT_I420");
+		return;
+	}
+
+	MediaCapture = NewObject<UPixelCaptureMediaCapture>();
+	MediaCapture->AddToRoot(); // prevent GC on this
+
+	MediaOutput = NewObject<UPixelCaptureMediaOuput>();
+	MediaCapture->SetMediaOutput(MediaOutput);
+	MediaCapture->SetFormat(Format);
+}
+
+FPixelCaptureCapturerMediaCapture::~FPixelCaptureCapturerMediaCapture()
+{
+	// We don't need to remove mediacapture from root if engine is shutting down
+	// as UE will already have killed all UObjects by this point.
+	if (!IsEngineExitRequested() && MediaCapture)
+	{
+		MediaCapture->RemoveFromRoot();
+	}
+}
+
+void FPixelCaptureCapturerMediaCapture::LateStartCapturer()
+{
+	OnFrameEndDelegateHandle = FCoreDelegates::OnBeginFrame.AddSP(AsShared(), &FPixelCaptureCapturerMediaCapture::InitializeMediaCapture);
+	MediaCapture->OnCaptureComplete.AddSP(AsShared(), &FPixelCaptureCapturerMediaCapture::EndProcess);
+}
+
+void FPixelCaptureCapturerMediaCapture::InitializeMediaCapture()
+{
+	// If we were bound to the OnFrameEnd delegate to ensure a frame was rendered before starting, then we can unset it here.
+	if (OnFrameEndDelegateHandle.IsSet())
+	{
+		FCoreDelegates::OnBeginFrame.Remove(OnFrameEndDelegateHandle.GetValue());
+		OnFrameEndDelegateHandle.Reset();
+	}
+
+	FMediaCaptureOptions CaptureOptions;
+	CaptureOptions.bSkipFrameWhenRunningExpensiveTasks = false;
+	CaptureOptions.OverrunAction = EMediaCaptureOverrunAction::Skip;
+	CaptureOptions.ResizeMethod = EMediaCaptureResizeMethod::ResizeInRenderPass;
+
+	FRHICaptureResourceDescription ResourceDescription;
+	ResourceDescription.PixelFormat = EPixelFormat::PF_B8G8R8A8;
+
+	MediaCapture->CaptureRHITexture(ResourceDescription, CaptureOptions);
+
+	bMediaCaptureInitialized = true;
+}
+
+IPixelCaptureOutputFrame* FPixelCaptureCapturerMediaCapture::CreateOutputBuffer(int32 InputWidth, int32 InputHeight)
+{
+	const int32 Width = InputWidth * Scale;
+	const int32 Height = InputHeight * Scale;
+
+	MediaOutput->SetRequestedSize({ Width, Height });
+
+	if (Format == PixelCaptureBufferFormat::FORMAT_RHI)
+	{
+		return new FPixelCaptureOutputFrameRHI(nullptr);
+	}
+	else if (Format == PixelCaptureBufferFormat::FORMAT_I420)
+	{
+		return new FPixelCaptureOutputFrameI420(nullptr);
+	}
+	else
+	{
+		UE_LOGFMT(LogPixelCapture, Error, "FPixelCaptureCapturerMediaCapture: Invalid pixel format. Expected either FORMAT_RHI or FORMAT_I420");
+		return nullptr;
+	}
+}
+
+void FPixelCaptureCapturerMediaCapture::BeginProcess(const IPixelCaptureInputFrame& InputFrame, IPixelCaptureOutputFrame* OutputBuffer)
+{
+	if (!bMediaCaptureInitialized)
+	{
+		// Early out as media capture is still initializing itself. We'll capture a later frame
+		EndProcess();
+		return;
+	}
+
+	checkf(InputFrame.GetType() == StaticCast<int32>(PixelCaptureBufferFormat::FORMAT_RHI), TEXT("Incorrect source frame coming into frame capture process."));
+	const FPixelCaptureInputFrameRHI& SourceFrame = StaticCast<const FPixelCaptureInputFrameRHI&>(InputFrame);
+
+	if (Format == PixelCaptureBufferFormat::FORMAT_RHI)
+	{
+		// If the source texture already matches the dimensions and pixelformat we're looking for
+		// we can just assign the input to the output and early out
+		const FRHITextureDesc& SourceDesc = SourceFrame.FrameTexture->GetDesc();
+		if (SourceDesc.Extent == MediaOutput->GetRequestedSize()
+			&& SourceDesc.Format == EPixelFormat::PF_B8G8R8A8
+#if PLATFORM_MAC
+			// Mac output textures have to have the CPUReadback flag, so if the input doesn't have it we can't do no copy process
+			&& EnumHasAnyFlags(SourceDesc.Flags, TexCreate_CPUReadback)
+#endif
+		)
+		{
+			static_cast<FPixelCaptureOutputFrameRHI*>(OutputBuffer)->SetFrameTexture(SourceFrame.FrameTexture);
+			EndProcess();
+			return;
+		}
+	}
+
+	MediaCapture->SetOutputFrame(OutputBuffer);
+
+	FRDGBuilder GraphBuilder(FRHICommandListImmediate::Get());
+	bool		bPassesAdded = MediaCapture->TryCaptureImmediate_RenderThread(GraphBuilder, SourceFrame.FrameTexture);
+	// Even if no passes are added, we still need to call Execute
+	GraphBuilder.Execute();
+
+	if (!bPassesAdded)
+	{
+		// RDG graph had no passes so we can manually call EndProcess()
+		EndProcess();
+	}
+}
\ No newline at end of file
diff --git a/Engine/Plugins/Media/PixelCapture/Source/PixelCapture/Private/PixelCapturePrivate.h b/Engine/Plugins/Media/PixelCapture/Source/PixelCapture/Private/PixelCapturePrivate.h
--- a/Engine/Plugins/Media/PixelCapture/Source/PixelCapture/Private/PixelCapturePrivate.h
+++ b/Engine/Plugins/Media/PixelCapture/Source/PixelCapture/Private/PixelCapturePrivate.h
@@ -4,5 +4,6 @@
 
 #include "HAL/IConsoleManager.h"
 #include "Logging/LogMacros.h"
+#include "Logging/StructuredLog.h"
 
 DECLARE_LOG_CATEGORY_EXTERN(LogPixelCapture, Log, All);
\ No newline at end of file
diff --git a/Engine/Plugins/Media/PixelCapture/Source/PixelCapture/Public/PixelCaptureOutputFrameI420.h b/Engine/Plugins/Media/PixelCapture/Source/PixelCapture/Public/PixelCaptureOutputFrameI420.h
--- a/Engine/Plugins/Media/PixelCapture/Source/PixelCapture/Public/PixelCaptureOutputFrameI420.h
+++ b/Engine/Plugins/Media/PixelCapture/Source/PixelCapture/Public/PixelCaptureOutputFrameI420.h
@@ -21,6 +21,7 @@ public:
 	virtual int32 GetHeight() const override { return I420Buffer->GetHeight(); }
 
 	TSharedPtr<FPixelCaptureBufferI420> GetI420Buffer() const { return I420Buffer; }
+	void SetI420Buffer(TSharedPtr<FPixelCaptureBufferI420> InI420Buffer) { I420Buffer = InI420Buffer; }
 
 private:
 	TSharedPtr<FPixelCaptureBufferI420> I420Buffer;
