[PixelCapture] QOL: Increate the MediaOuput's NumberOfTextureBuffers to 8 to decrese chance of overrun occuring

diff --git a/Engine/Plugins/Media/PixelCapture/Source/PixelCapture/Private/PixelCaptureCapturerMediaCapture.cpp b/Engine/Plugins/Media/PixelCapture/Source/PixelCapture/Private/PixelCaptureCapturerMediaCapture.cpp
--- a/Engine/Plugins/Media/PixelCapture/Source/PixelCapture/Private/PixelCaptureCapturerMediaCapture.cpp
+++ b/Engine/Plugins/Media/PixelCapture/Source/PixelCapture/Private/PixelCaptureCapturerMediaCapture.cpp
@@ -192,6 +192,12 @@ FPixelCaptureCapturerMediaCapture::FPixelCaptureCapturerMediaCapture(float InSca
 	MediaCapture->AddToRoot(); // prevent GC on this
 
 	MediaOutput = NewObject<UPixelCaptureMediaOuput>();
+	// Note the number of texture buffers is how many textures we have in reserve to copy into while we wait for other captures to complete
+	// On slower hardware this number needs to be bigger. Testing on AWS T4 GPU's (which are sort of like min-spec for PS) we determined
+	// the default number (4) is too low and will cause media capture to regularly overrun (which results in either a skipped frame or a
+	// GPU flush depending on the EMediaCaptureOverrunAction option below). After testing, it was found that 8 textures (the max),
+	// reduced overruns to infrequent levels on the AWS T4 GPU.
+	MediaOutput->NumberOfTextureBuffers = 8;
 	MediaCapture->SetMediaOutput(MediaOutput);
 	MediaCapture->SetFormat(Format);
 }
@@ -224,7 +230,7 @@ void FPixelCaptureCapturerMediaCapture::InitializeMediaCapture()
 	FMediaCaptureOptions CaptureOptions;
 	CaptureOptions.bSkipFrameWhenRunningExpensiveTasks = false;
 	CaptureOptions.OverrunAction = EMediaCaptureOverrunAction::Skip;
-	CaptureOptions.ResizeMethod = EMediaCaptureResizeMethod::ResizeInRenderPass;
+	CaptureOptions.ResizeMethod = EMediaCaptureResizeMethod::None;
 
 	FRHICaptureResourceDescription ResourceDescription;
 	ResourceDescription.PixelFormat = EPixelFormat::PF_B8G8R8A8;
