[PixelCapture] Fix: Race condition with Frame lifecycle delegate. Move to just posting an async task instead

diff --git a/Engine/Plugins/Media/PixelCapture/Source/PixelCapture/Internal/PixelCaptureCapturerMediaCapture.h b/Engine/Plugins/Media/PixelCapture/Source/PixelCapture/Internal/PixelCaptureCapturerMediaCapture.h
--- a/Engine/Plugins/Media/PixelCapture/Source/PixelCapture/Internal/PixelCaptureCapturerMediaCapture.h
+++ b/Engine/Plugins/Media/PixelCapture/Source/PixelCapture/Internal/PixelCaptureCapturerMediaCapture.h
@@ -114,13 +114,10 @@ private:
 
 	FPixelCaptureCapturerMediaCapture(float InScale, int32 InFormat);
 	void InitializeMediaCapture();
-	void LateStartCapturer();
 
 	UPROPERTY(Transient)
 	TObjectPtr<UPixelCaptureMediaCapture> MediaCapture = nullptr;
 
 	UPROPERTY(Transient)
 	TObjectPtr<UPixelCaptureMediaOuput> MediaOutput = nullptr;
-
-	TOptional<FDelegateHandle> OnFrameEndDelegateHandle;
 };
diff --git a/Engine/Plugins/Media/PixelCapture/Source/PixelCapture/Private/PixelCaptureCapturerMediaCapture.cpp b/Engine/Plugins/Media/PixelCapture/Source/PixelCapture/Private/PixelCaptureCapturerMediaCapture.cpp
--- a/Engine/Plugins/Media/PixelCapture/Source/PixelCapture/Private/PixelCaptureCapturerMediaCapture.cpp
+++ b/Engine/Plugins/Media/PixelCapture/Source/PixelCapture/Private/PixelCaptureCapturerMediaCapture.cpp
@@ -173,7 +173,14 @@ ETextureCreateFlags UPixelCaptureMediaCapture::GetOutputTextureFlags() const
 TSharedPtr<FPixelCaptureCapturerMediaCapture> FPixelCaptureCapturerMediaCapture::Create(float InScale, int32 InFormat)
 {
 	TSharedPtr<FPixelCaptureCapturerMediaCapture> Capturer = TSharedPtr<FPixelCaptureCapturerMediaCapture>(new FPixelCaptureCapturerMediaCapture(InScale, InFormat));
-	Capturer->LateStartCapturer();
+
+	TWeakPtr<FPixelCaptureCapturerMediaCapture> WeakCapturer = Capturer;
+	AsyncTask(ENamedThreads::GameThread, [WeakCapturer]() {
+		if (TSharedPtr<FPixelCaptureCapturerMediaCapture> PinnedCapturer = WeakCapturer.Pin())
+		{
+			PinnedCapturer->InitializeMediaCapture();
+		}
+	});
 
 	return Capturer;
 }
@@ -212,20 +219,9 @@ FPixelCaptureCapturerMediaCapture::~FPixelCaptureCapturerMediaCapture()
 	}
 }
 
-void FPixelCaptureCapturerMediaCapture::LateStartCapturer()
-{
-	OnFrameEndDelegateHandle = FCoreDelegates::OnBeginFrame.AddSP(AsShared(), &FPixelCaptureCapturerMediaCapture::InitializeMediaCapture);
-	MediaCapture->OnCaptureComplete.AddSP(AsShared(), &FPixelCaptureCapturerMediaCapture::EndProcess);
-}
-
 void FPixelCaptureCapturerMediaCapture::InitializeMediaCapture()
 {
-	// If we were bound to the OnFrameEnd delegate to ensure a frame was rendered before starting, then we can unset it here.
-	if (OnFrameEndDelegateHandle.IsSet())
-	{
-		FCoreDelegates::OnBeginFrame.Remove(OnFrameEndDelegateHandle.GetValue());
-		OnFrameEndDelegateHandle.Reset();
-	}
+	MediaCapture->OnCaptureComplete.AddSP(AsShared(), &FPixelCaptureCapturerMediaCapture::EndProcess);
 
 	FMediaCaptureOptions CaptureOptions;
 	CaptureOptions.bSkipFrameWhenRunningExpensiveTasks = false;
