[PS2] Fix: Poor VPX performance, freeze frame bugs, simulcast capture not working

diff --git a/Engine/Plugins/Media/PixelCapture/Source/PixelCapture/Private/PixelCaptureCapturerRHIToI420CPU.cpp b/Engine/Plugins/Media/PixelCapture/Source/PixelCapture/Private/PixelCaptureCapturerRHIToI420CPU.cpp
--- a/Engine/Plugins/Media/PixelCapture/Source/PixelCapture/Private/PixelCaptureCapturerRHIToI420CPU.cpp
+++ b/Engine/Plugins/Media/PixelCapture/Source/PixelCapture/Private/PixelCaptureCapturerRHIToI420CPU.cpp
@@ -25,6 +25,15 @@ FPixelCaptureCapturerRHIToI420CPU::~FPixelCaptureCapturerRHIToI420CPU()
 {
 }
 
+IPixelCaptureOutputFrame* FPixelCaptureCapturerRHIToI420CPU::CreateOutputBuffer(int32 InputWidth, int32 InputHeight)
+{
+	const int32 Width = InputWidth * Scale;
+	const int32 Height = InputHeight * Scale;
+	return new FPixelCaptureOutputFrameI420(MakeShared<FPixelCaptureBufferI420>(Width, Height));
+}
+
+#if 0
+// Old code path that uses a FRHIGPUTextureReadback. This seems to induce extra latency so is disabled for now
 void FPixelCaptureCapturerRHIToI420CPU::Initialize(int32 InputWidth, int32 InputHeight)
 {
 	const int32 Width = InputWidth * Scale;
@@ -53,13 +62,6 @@ void FPixelCaptureCapturerRHIToI420CPU::Initialize(int32 InputWidth, int32 Input
 	FPixelCaptureCapturer::Initialize(InputWidth, InputHeight);
 }
 
-IPixelCaptureOutputFrame* FPixelCaptureCapturerRHIToI420CPU::CreateOutputBuffer(int32 InputWidth, int32 InputHeight)
-{
-	const int32 Width = InputWidth * Scale;
-	const int32 Height = InputHeight * Scale;
-	return new FPixelCaptureOutputFrameI420(MakeShared<FPixelCaptureBufferI420>(Width, Height));
-}
-
 void FPixelCaptureCapturerRHIToI420CPU::BeginProcess(const IPixelCaptureInputFrame& InputFrame, IPixelCaptureOutputFrame* OutputBuffer)
 {
 	checkf(InputFrame.GetType() == StaticCast<int32>(PixelCaptureBufferFormat::FORMAT_RHI), TEXT("Incorrect source frame coming into frame capture process."));
@@ -131,3 +133,79 @@ void FPixelCaptureCapturerRHIToI420CPU::OnRHIStageComplete(IPixelCaptureOutputFr
 	MarkCPUWorkEnd();
 	EndProcess();
 }
+
+#else
+
+void FPixelCaptureCapturerRHIToI420CPU::Initialize(int32 InputWidth, int32 InputHeight)
+{
+	const int32 Width = InputWidth * Scale;
+	const int32 Height = InputHeight * Scale;
+
+	FRHITextureCreateDesc StagingTextureDesc =
+		FRHITextureCreateDesc::Create2D(TEXT("FPixelCaptureCapturerRHIToI420CPU StagingTexture"), Width, Height, EPixelFormat::PF_B8G8R8A8)
+			.SetClearValue(FClearValueBinding::None)
+			.SetFlags(ETextureCreateFlags::RenderTargetable)
+			.SetInitialState(ERHIAccess::CopySrc)
+			.DetermineInititialState();
+
+	if (RHIGetInterfaceType() == ERHIInterfaceType::Vulkan)
+	{
+		StagingTextureDesc.AddFlags(ETextureCreateFlags::External);
+	}
+	else
+	{
+		StagingTextureDesc.AddFlags(ETextureCreateFlags::Shared);
+	}
+
+	StagingTexture = RHICreateTexture(StagingTextureDesc);
+
+	FRHITextureCreateDesc ReadbackTextureDesc = FRHITextureCreateDesc::Create2D(TEXT("FPixelCaptureCapturerRHIToI420CPU ReadbackTexture"), Width, Height, EPixelFormat::PF_B8G8R8A8)
+													.SetClearValue(FClearValueBinding::None)
+													.SetFlags(ETextureCreateFlags::CPUReadback | ETextureCreateFlags::HideInVisualizeTexture)
+													.SetInitialState(ERHIAccess::CopyDest)
+													.DetermineInititialState();
+
+	ReadbackTexture = RHICreateTexture(ReadbackTextureDesc);
+
+	FPixelCaptureCapturer::Initialize(InputWidth, InputHeight);
+}
+
+void FPixelCaptureCapturerRHIToI420CPU::BeginProcess(const IPixelCaptureInputFrame& InputFrame, IPixelCaptureOutputFrame* OutputBuffer)
+{
+	checkf(InputFrame.GetType() == StaticCast<int32>(PixelCaptureBufferFormat::FORMAT_RHI), TEXT("Incorrect source frame coming into frame capture process."));
+
+	const FPixelCaptureInputFrameRHI& RHISourceFrame = StaticCast<const FPixelCaptureInputFrameRHI&>(InputFrame);
+	FTextureRHIRef					  SourceTexture = RHISourceFrame.FrameTexture;
+
+	FRHICommandListImmediate& RHICmdList = FRHICommandListExecutor::GetImmediateCommandList();
+	RHICmdList.Transition(FRHITransitionInfo(SourceTexture, ERHIAccess::Unknown, ERHIAccess::CopySrc));
+	RHICmdList.Transition(FRHITransitionInfo(StagingTexture, ERHIAccess::CopySrc, ERHIAccess::CopyDest));
+	CopyTexture(RHICmdList, SourceTexture, StagingTexture, nullptr);
+	RHICmdList.Transition(FRHITransitionInfo(StagingTexture, ERHIAccess::CopyDest, ERHIAccess::CopySrc));
+	// Format and extent match so this is just a simple copy
+	CopyTexture(RHICmdList, StagingTexture, ReadbackTexture, nullptr);
+
+	void* ReadbackPointer = nullptr;
+	int32 ReadbackWidth, ReadbackHeight;
+	// Nullptr fence forces a flush
+	RHICmdList.MapStagingSurface(ReadbackTexture, nullptr, ReadbackPointer, ReadbackWidth, ReadbackHeight);
+
+	FPixelCaptureOutputFrameI420*		OutputI420Buffer = StaticCast<FPixelCaptureOutputFrameI420*>(OutputBuffer);
+	TSharedPtr<FPixelCaptureBufferI420> I420Buffer = OutputI420Buffer->GetI420Buffer();
+	libyuv::ARGBToI420(
+		static_cast<uint8*>(ReadbackPointer),
+		ReadbackWidth * 4,
+		I420Buffer->GetMutableDataY(),
+		I420Buffer->GetStrideY(),
+		I420Buffer->GetMutableDataU(),
+		I420Buffer->GetStrideUV(),
+		I420Buffer->GetMutableDataV(),
+		I420Buffer->GetStrideUV(),
+		I420Buffer->GetWidth(),
+		I420Buffer->GetHeight());
+
+	RHICmdList.UnmapStagingSurface(ReadbackTexture);
+
+	EndProcess();
+}
+#endif
\ No newline at end of file
diff --git a/Engine/Plugins/Media/PixelCapture/Source/PixelCapture/Public/PixelCaptureCapturerRHIToI420CPU.h b/Engine/Plugins/Media/PixelCapture/Source/PixelCapture/Public/PixelCaptureCapturerRHIToI420CPU.h
--- a/Engine/Plugins/Media/PixelCapture/Source/PixelCapture/Public/PixelCaptureCapturerRHIToI420CPU.h
+++ b/Engine/Plugins/Media/PixelCapture/Source/PixelCapture/Public/PixelCaptureCapturerRHIToI420CPU.h
@@ -7,10 +7,10 @@
 #include "RHIGPUReadback.h"
 
 /**
- * A basic capturer that will capture RHI texture frames to I420 buffers utilizing cpu functions.
- * Involves CPU readback of GPU textures and processing of that readback data.
- * Input: FPixelCaptureInputFrameRHI
- * Output: FPixelCaptureOutputFrameI420
+ * A basic capturer that will capture RHI texture frames to I420 buffers
+ * utilizing cpu functions. Involves CPU readback of GPU textures and processing
+ * of that readback data. Input: FPixelCaptureInputFrameRHI Output:
+ * FPixelCaptureOutputFrameI420
  */
 class PIXELCAPTURE_API FPixelCaptureCapturerRHIToI420CPU : public FPixelCaptureCapturer, public TSharedFromThis<FPixelCaptureCapturerRHIToI420CPU>
 {
@@ -23,18 +23,23 @@ public:
 	virtual ~FPixelCaptureCapturerRHIToI420CPU();
 
 protected:
-	virtual FString GetCapturerName() const override { return "RHIToI420CPU"; }
-	virtual void Initialize(int32 InputWidth, int32 InputHeight) override;
+	virtual FString					  GetCapturerName() const override { return "RHIToI420CPU"; }
+	virtual void					  Initialize(int32 InputWidth, int32 InputHeight) override;
 	virtual IPixelCaptureOutputFrame* CreateOutputBuffer(int32 InputWidth, int32 InputHeight) override;
-	virtual void BeginProcess(const IPixelCaptureInputFrame& InputFrame, IPixelCaptureOutputFrame* OutputBuffer) override;
+	virtual void					  BeginProcess(const IPixelCaptureInputFrame& InputFrame, IPixelCaptureOutputFrame* OutputBuffer) override;
 
 private:
 	float Scale = 1.0f;
 
-	TSharedPtr<FRHIGPUTextureReadback> TextureReader;
 	FTextureRHIRef StagingTexture;
 
 	FPixelCaptureCapturerRHIToI420CPU(float InScale);
+
+#if 0
+	TSharedPtr<FRHIGPUTextureReadback> TextureReader;
 	void OnRHIStageComplete(IPixelCaptureOutputFrame* OutputBuffer);
 	void CheckComplete(IPixelCaptureOutputFrame* OutputBuffer);
+#else
+	FTextureRHIRef ReadbackTexture;
+#endif
 };
diff --git a/Engine/Plugins/Media/PixelStreaming2/Source/PixelStreaming2/Private/VideoCapturer.cpp b/Engine/Plugins/Media/PixelStreaming2/Source/PixelStreaming2/Private/VideoCapturer.cpp
--- a/Engine/Plugins/Media/PixelStreaming2/Source/PixelStreaming2/Private/VideoCapturer.cpp
+++ b/Engine/Plugins/Media/PixelStreaming2/Source/PixelStreaming2/Private/VideoCapturer.cpp
@@ -8,6 +8,8 @@
 #include "PixelCaptureCapturerI420.h"
 #include "PixelCaptureCapturerI420ToRHI.h"
 #include "PixelCaptureCapturerNV12ToRHI.h"
+#include "PixelCaptureCapturerRHIRDG.h"
+#include "PixelCaptureCapturerRHIToI420CPU.h"
 #include "PixelStreaming2PluginSettings.h"
 #include "Logging.h"
 #include "PixelStreaming2Trace.h"
@@ -133,7 +135,16 @@ namespace UE::PixelStreaming2
 	{
 		if (LastFrameType == PixelCaptureBufferFormat::FORMAT_RHI)
 		{
-			return FPixelCaptureCapturerMediaCapture::Create(FinalScale, FinalFormat);
+			switch (FinalFormat)
+			{
+				case PixelCaptureBufferFormat::FORMAT_RHI:
+					return FPixelCaptureCapturerRHIRDG::Create(FinalScale);
+				case PixelCaptureBufferFormat::FORMAT_I420:
+					return FPixelCaptureCapturerRHIToI420CPU::Create(FinalScale);
+				default:
+					UE_LOGFMT(LogPixelStreaming2, Error, "Unsupported final format ({0}) for RHI input format", FinalFormat);
+					return nullptr;
+			}
 		}
 		else if (LastFrameType == PixelCaptureBufferFormat::FORMAT_I420)
 		{
