[PixelCapture] Fix when multiple FrameCapturers catures a frame, causing a data race inside delegates in Pixel Capture in OnComplete.

Fix is to use thread safe delegate.

diff --git a/Engine/Plugins/Media/PixelCapture/Source/PixelCapture/Public/PixelCaptureCapturerMultiFormat.h b/Engine/Plugins/Media/PixelCapture/Source/PixelCapture/Public/PixelCaptureCapturerMultiFormat.h
--- a/Engine/Plugins/Media/PixelCapture/Source/PixelCapture/Public/PixelCaptureCapturerMultiFormat.h
+++ b/Engine/Plugins/Media/PixelCapture/Source/PixelCapture/Public/PixelCaptureCapturerMultiFormat.h
@@ -84,7 +84,7 @@ public:
 	/**
 	 * Listen on this to be notified when a frame completes all capture formats/layers.
 	 */
-	DECLARE_MULTICAST_DELEGATE(FOnComplete);
+	DECLARE_TS_MULTICAST_DELEGATE(FOnComplete);
 	FOnComplete OnComplete;
 
 protected:
