[PixelCapture] Vsync issue when using vpx with Pixel Streaming.

This fix re-enables code that uses a fence when retrieving the frame from the GPU.

diff --git a/Engine/Plugins/Media/PixelCapture/Source/PixelCapture/Private/PixelCaptureCapturerRHIToI420CPU.cpp b/Engine/Plugins/Media/PixelCapture/Source/PixelCapture/Private/PixelCaptureCapturerRHIToI420CPU.cpp
--- a/Engine/Plugins/Media/PixelCapture/Source/PixelCapture/Private/PixelCaptureCapturerRHIToI420CPU.cpp
+++ b/Engine/Plugins/Media/PixelCapture/Source/PixelCapture/Private/PixelCaptureCapturerRHIToI420CPU.cpp
@@ -32,8 +32,8 @@ IPixelCaptureOutputFrame* FPixelCaptureCapturerRHIToI420CPU::CreateOutputBuffer(
 	return new FPixelCaptureOutputFrameI420(MakeShared<FPixelCaptureBufferI420>(Width, Height));
 }
 
-#if 0
-// Old code path that uses a FRHIGPUTextureReadback. This seems to induce extra latency so is disabled for now
+#if 1
+// Old code path that uses a FRHIGPUTextureReadback. This seems to induce extra latency but does not introduce screen tearing
 void FPixelCaptureCapturerRHIToI420CPU::Initialize(int32 InputWidth, int32 InputHeight)
 {
 	const int32 Width = InputWidth * Scale;
@@ -134,7 +134,8 @@ void FPixelCaptureCapturerRHIToI420CPU::OnRHIStageComplete(IPixelCaptureOutputFr
 	EndProcess();
 }
 
-#else
+// TODO (Eden.Harris) This has less latency but introduces screen tearing so is disabled. RTCP-7778 to fix the screen tearing without effecting latency.
+#else	
 
 void FPixelCaptureCapturerRHIToI420CPU::Initialize(int32 InputWidth, int32 InputHeight)
 {
diff --git a/Engine/Plugins/Media/PixelCapture/Source/PixelCapture/Public/PixelCaptureCapturerRHIToI420CPU.h b/Engine/Plugins/Media/PixelCapture/Source/PixelCapture/Public/PixelCaptureCapturerRHIToI420CPU.h
--- a/Engine/Plugins/Media/PixelCapture/Source/PixelCapture/Public/PixelCaptureCapturerRHIToI420CPU.h
+++ b/Engine/Plugins/Media/PixelCapture/Source/PixelCapture/Public/PixelCaptureCapturerRHIToI420CPU.h
@@ -35,7 +35,7 @@ private:
 
 	FPixelCaptureCapturerRHIToI420CPU(float InScale);
 
-#if 0
+#if 1
 	TSharedPtr<FRHIGPUTextureReadback> TextureReader;
 	void OnRHIStageComplete(IPixelCaptureOutputFrame* OutputBuffer);
 	void CheckComplete(IPixelCaptureOutputFrame* OutputBuffer);
