[PS2] Fix: Frame stuttering with MediaIO due to not waiting for texture copy completion

diff --git a/Engine/Plugins/Media/PixelCapture/Source/PixelCapture/Internal/PixelCaptureCapturerMediaCapture.h b/Engine/Plugins/Media/PixelCapture/Source/PixelCapture/Internal/PixelCaptureCapturerMediaCapture.h
--- a/Engine/Plugins/Media/PixelCapture/Source/PixelCapture/Internal/PixelCaptureCapturerMediaCapture.h
+++ b/Engine/Plugins/Media/PixelCapture/Source/PixelCapture/Internal/PixelCaptureCapturerMediaCapture.h
@@ -2,6 +2,7 @@
 
 #pragma once
 
+#include "Containers/Deque.h"
 #include "MediaCapture.h"
 #include "MediaOutput.h"
 #include "PixelCaptureBufferFormat.h"
@@ -81,10 +82,13 @@ protected:
 	}
 
 	virtual ETextureCreateFlags GetOutputTextureFlags() const override;
+	virtual void				WaitForGPU(FRHITexture* InRHITexture) override;
 
 private:
 	IPixelCaptureOutputFrame* OutputFrame = nullptr;
 	int32					  Format = PixelCaptureBufferFormat::FORMAT_UNKNOWN;
+
+	TDeque<FGPUFenceRHIRef> Fences;
 };
 
 /**
diff --git a/Engine/Plugins/Media/PixelCapture/Source/PixelCapture/Private/PixelCaptureCapturerMediaCapture.cpp b/Engine/Plugins/Media/PixelCapture/Source/PixelCapture/Private/PixelCaptureCapturerMediaCapture.cpp
--- a/Engine/Plugins/Media/PixelCapture/Source/PixelCapture/Private/PixelCaptureCapturerMediaCapture.cpp
+++ b/Engine/Plugins/Media/PixelCapture/Source/PixelCapture/Private/PixelCaptureCapturerMediaCapture.cpp
@@ -72,6 +72,9 @@ void UPixelCaptureMediaCapture::OnCustomCapture_RenderingThread(
 	FVector2D											   CropU,
 	FVector2D											   CropV)
 {
+	FGPUFenceRHIRef Fence = GDynamicRHI->RHICreateGPUFence(TEXT("UPixelCaptureMediaCapture Fence"));
+	Fences.PushLast(Fence);
+
 	bool bRequiresFormatConversion = InSourceTexture->Desc.Format != OutputTexture->Desc.Format;
 	if (InSourceTexture->Desc.Format == OutputTexture->Desc.Format
 		&& InSourceTexture->Desc.Extent.X == OutputTexture->Desc.Extent.X
@@ -84,8 +87,6 @@ void UPixelCaptureMediaCapture::OnCustomCapture_RenderingThread(
 			InSourceTexture,
 			OutputTexture,
 			FRDGDrawTextureInfo());
-
-		return;
 	}
 	else
 	{
@@ -148,6 +149,15 @@ void UPixelCaptureMediaCapture::OnCustomCapture_RenderingThread(
 			FRDGDrawTextureInfo());
 #endif
 	}
+
+	// clang-format off
+	GraphBuilder.AddPass(
+		RDG_EVENT_NAME("UPixelCaptureMediaCapture WriteFence"), 
+		ERDGPassFlags::NeverCull,								
+		[Fence](FRDGAsyncTask, FRHICommandList& RHICmdList) {
+			RHICmdList.WriteGPUFence(Fence);
+		});
+	// clang-format on
 }
 
 ETextureCreateFlags UPixelCaptureMediaCapture::GetOutputTextureFlags() const
@@ -170,6 +180,21 @@ ETextureCreateFlags UPixelCaptureMediaCapture::GetOutputTextureFlags() const
 #endif
 }
 
+void UPixelCaptureMediaCapture::WaitForGPU(FRHITexture* InRHITexture)
+{
+	FGPUFenceRHIRef Fence;
+	if (!Fences.TryPopFirst(Fence))
+	{
+		return;
+	}
+
+	while (!Fence->Poll())
+	{
+		constexpr float SleepTimeSeconds = 50 * 1E-6;
+		FPlatformProcess::SleepNoStats(SleepTimeSeconds);
+	}
+}
+
 TSharedPtr<FPixelCaptureCapturerMediaCapture> FPixelCaptureCapturerMediaCapture::Create(float InScale, int32 InFormat)
 {
 	TSharedPtr<FPixelCaptureCapturerMediaCapture> Capturer = TSharedPtr<FPixelCaptureCapturerMediaCapture>(new FPixelCaptureCapturerMediaCapture(InScale, InFormat));
@@ -227,6 +252,10 @@ void FPixelCaptureCapturerMediaCapture::InitializeMediaCapture()
 	CaptureOptions.bSkipFrameWhenRunningExpensiveTasks = false;
 	CaptureOptions.OverrunAction = EMediaCaptureOverrunAction::Skip;
 	CaptureOptions.ResizeMethod = EMediaCaptureResizeMethod::ResizeInCapturePass;
+	CaptureOptions.bAutostopOnCapture = true;
+	// This value needs to be >= 1 in order to be a valid configuration
+	// but it isn't actually used with TryCaptureImmediate_RenderThread
+	CaptureOptions.NumberOfFramesToCapture = 1; 
 
 	FRHICaptureResourceDescription ResourceDescription;
 	ResourceDescription.PixelFormat = EPixelFormat::PF_B8G8R8A8;
diff --git a/Engine/Plugins/Media/PixelStreaming2/Source/PixelStreaming2/Internal/PixelStreaming2MediaIOCapture.h b/Engine/Plugins/Media/PixelStreaming2/Source/PixelStreaming2/Internal/PixelStreaming2MediaIOCapture.h
--- a/Engine/Plugins/Media/PixelStreaming2/Source/PixelStreaming2/Internal/PixelStreaming2MediaIOCapture.h
+++ b/Engine/Plugins/Media/PixelStreaming2/Source/PixelStreaming2/Internal/PixelStreaming2MediaIOCapture.h
@@ -2,6 +2,7 @@
 
 #pragma once
 
+#include "Containers/Deque.h"
 #include "MediaCapture.h"
 #include "Slate/SceneViewport.h"
 #include "VideoProducer.h"
@@ -43,6 +44,7 @@ public:
 	virtual void StopCaptureImpl(bool bAllowPendingFrameToBeProcess) override;
 	// We override the texture flags because on Mac we want the texture to have the CPU_Readback flag
 	virtual ETextureCreateFlags GetOutputTextureFlags() const override;
+	virtual void WaitForGPU(FRHITexture* InRHITexture) override;
 	//~ End UMediaCapture interface
 
 	TSharedPtr<FSceneViewport>				GetViewport() const { return SceneViewport.Pin(); }
@@ -70,4 +72,6 @@ private:
 
 	/* Whether we want the UMediaCapture to read back to frame into cpu memory or not. */
 	bool bDoGPUCopy = true;
+
+	TDeque<FGPUFenceRHIRef> Fences;
 };
\ No newline at end of file
diff --git a/Engine/Plugins/Media/PixelStreaming2/Source/PixelStreaming2/Private/PixelStreaming2MediaIOCapture.cpp b/Engine/Plugins/Media/PixelStreaming2/Source/PixelStreaming2/Private/PixelStreaming2MediaIOCapture.cpp
--- a/Engine/Plugins/Media/PixelStreaming2/Source/PixelStreaming2/Private/PixelStreaming2MediaIOCapture.cpp
+++ b/Engine/Plugins/Media/PixelStreaming2/Source/PixelStreaming2/Private/PixelStreaming2MediaIOCapture.cpp
@@ -30,6 +30,9 @@ void UPixelStreaming2MediaIOCapture::OnCustomCapture_RenderingThread(
 	FVector2D											   CropU,
 	FVector2D											   CropV)
 {
+	FGPUFenceRHIRef Fence = GDynamicRHI->RHICreateGPUFence(TEXT("UPixelStreaming2MediaIOCapture Fence"));
+	Fences.PushLast(Fence);
+
 	bool bRequiresFormatConversion = InSourceTexture->Desc.Format != OutputTexture->Desc.Format;
 	if (InSourceTexture->Desc.Format == OutputTexture->Desc.Format
 		&& InSourceTexture->Desc.Extent.X == OutputTexture->Desc.Extent.X
@@ -42,8 +45,6 @@ void UPixelStreaming2MediaIOCapture::OnCustomCapture_RenderingThread(
 			InSourceTexture,
 			OutputTexture,
 			FRDGDrawTextureInfo());
-
-		return;
 	}
 	else
 	{
@@ -106,6 +107,15 @@ void UPixelStreaming2MediaIOCapture::OnCustomCapture_RenderingThread(
 			FRDGDrawTextureInfo());
 #endif
 	}
+
+	// clang-format off
+	GraphBuilder.AddPass(
+		RDG_EVENT_NAME("UPixelStreaming2MediaIOCapture WriteFence"), 
+		ERDGPassFlags::NeverCull,									 
+		[Fence](FRDGAsyncTask, FRHICommandList& RHICmdList) {
+			RHICmdList.WriteGPUFence(Fence);
+		});
+	// clang-format on
 }
 
 bool UPixelStreaming2MediaIOCapture::InitializeCapture()
@@ -144,6 +154,21 @@ ETextureCreateFlags UPixelStreaming2MediaIOCapture::GetOutputTextureFlags() cons
 #endif
 }
 
+void UPixelStreaming2MediaIOCapture::WaitForGPU(FRHITexture* InRHITexture)
+{
+	FGPUFenceRHIRef Fence;
+	if (!Fences.TryPopFirst(Fence))
+	{
+		return;
+	}
+
+	while (!Fence->Poll())
+	{
+		constexpr float SleepTimeSeconds = 50 * 1E-6;
+		FPlatformProcess::SleepNoStats(SleepTimeSeconds);
+	}
+}
+
 bool UPixelStreaming2MediaIOCapture::PostInitializeCaptureViewport(TSharedPtr<FSceneViewport>& InSceneViewport)
 {
 	SceneViewport = TWeakPtr<FSceneViewport>(InSceneViewport);
