[PixelCapture] Fix: Allow PixelCapture to queue multiple capture processes

diff --git a/Engine/Plugins/Media/PixelCapture/Source/PixelCapture/Internal/PixelCaptureCapturerMediaCapture.h b/Engine/Plugins/Media/PixelCapture/Source/PixelCapture/Internal/PixelCaptureCapturerMediaCapture.h
--- a/Engine/Plugins/Media/PixelCapture/Source/PixelCapture/Internal/PixelCaptureCapturerMediaCapture.h
+++ b/Engine/Plugins/Media/PixelCapture/Source/PixelCapture/Internal/PixelCaptureCapturerMediaCapture.h
@@ -31,9 +31,14 @@ class PIXELCAPTURE_API UPixelCaptureMediaCapture : public UMediaCapture
 	GENERATED_BODY()
 
 public:
-	void SetOutputFrame(IPixelCaptureOutputFrame* InOutputFrame)
+	void AddOutputFrame(TSharedPtr<IPixelCaptureOutputFrame> InOutputFrame)
 	{
-		OutputFrame = InOutputFrame;
+		OutputFrames.PushLast(InOutputFrame);
+	}
+
+	void RemoveLastOutputFrame()
+	{
+		OutputFrames.PopLast();
 	}
 
 	void SetFormat(int32 InFormat)
@@ -41,7 +46,7 @@ public:
 		Format = InFormat;
 	}
 
-	DECLARE_EVENT(UPixelCaptureMediaCapture, FOnCaptureComplete)
+	DECLARE_EVENT_OneParam(UPixelCaptureMediaCapture, FOnCaptureComplete, TSharedPtr<IPixelCaptureOutputFrame>);
 	FOnCaptureComplete OnCaptureComplete;
 
 protected:
@@ -85,10 +90,11 @@ protected:
 	virtual void				WaitForGPU(FRHITexture* InRHITexture) override;
 
 private:
-	IPixelCaptureOutputFrame* OutputFrame = nullptr;
 	int32					  Format = PixelCaptureBufferFormat::FORMAT_UNKNOWN;
 
 	TDeque<FGPUFenceRHIRef> Fences;
+
+	TDeque<TSharedPtr<IPixelCaptureOutputFrame>> OutputFrames;
 };
 
 /**
@@ -107,9 +113,9 @@ public:
 	virtual ~FPixelCaptureCapturerMediaCapture() override;
 
 protected:
-	virtual FString					  GetCapturerName() const override { return "FPixelCaptureCapturerMediaCapture Copy"; }
+	virtual FString					  GetCapturerName() const override { return "MediaCapture Copy"; }
 	virtual IPixelCaptureOutputFrame* CreateOutputBuffer(int32 InputWidth, int32 InputHeight) override;
-	virtual void					  BeginProcess(const IPixelCaptureInputFrame& InputFrame, IPixelCaptureOutputFrame* OutputBuffer) override;
+	virtual void					  BeginProcess(const IPixelCaptureInputFrame& InputFrame, TSharedPtr<IPixelCaptureOutputFrame> OutputBuffer) override;
 
 private:
 	float Scale = 1.0f;
diff --git a/Engine/Plugins/Media/PixelCapture/Source/PixelCapture/Private/OuputFrameBuffer.cpp b/Engine/Plugins/Media/PixelCapture/Source/PixelCapture/Private/OuputFrameBuffer.cpp
--- a/Engine/Plugins/Media/PixelCapture/Source/PixelCapture/Private/OuputFrameBuffer.cpp
+++ b/Engine/Plugins/Media/PixelCapture/Source/PixelCapture/Private/OuputFrameBuffer.cpp
@@ -2,6 +2,8 @@
 
 #include "OutputFrameBuffer.h"
 
+#include "PixelCapturePrivate.h"
+
 namespace UE::PixelCapture
 {
 	FOutputFrameBuffer::FOutputFrameBuffer()
@@ -58,9 +60,21 @@ namespace UE::PixelCapture
 		return BufferRing[ProduceIndex];
 	}
 
-	void FOutputFrameBuffer::ReleaseProduceBuffer()
+	bool FOutputFrameBuffer::ReleaseProduceBuffer(TSharedPtr<IPixelCaptureOutputFrame> OutputBuffer)
 	{
-		ConsumeIndex = StaticCast<int32>(ProduceIndex);
+		for (int i = 0; i < BufferRing.Num(); ++i)
+		{
+			if (BufferRing[i] == OutputBuffer)
+			{
+				// Set the consumer index to the most recently available output buffer.
+				// This enables multiple produce buffers to be mid capture only consume
+				// the most recently finished output buffer
+				ConsumeIndex = i;
+				return true;
+			}
+		}
+		UE_LOG(LogPixelCapture, Warning, TEXT("Failed to release produce buffer."));
+		return false;
 	}
 
 	TSharedPtr<IPixelCaptureOutputFrame> FOutputFrameBuffer::GetConsumeBuffer()
diff --git a/Engine/Plugins/Media/PixelCapture/Source/PixelCapture/Private/OutputFrameBuffer.h b/Engine/Plugins/Media/PixelCapture/Source/PixelCapture/Private/OutputFrameBuffer.h
--- a/Engine/Plugins/Media/PixelCapture/Source/PixelCapture/Private/OutputFrameBuffer.h
+++ b/Engine/Plugins/Media/PixelCapture/Source/PixelCapture/Private/OutputFrameBuffer.h
@@ -42,7 +42,7 @@ namespace UE::PixelCapture
 		 * Indicates that the previously acquired produce buffer has been filled and
 		 * ready to consume. Updates the consume position to this buffer.
 		 */
-		void ReleaseProduceBuffer();
+		bool ReleaseProduceBuffer(TSharedPtr<IPixelCaptureOutputFrame> OutputBuffer);
 
 		/**
 		 * Gets the newest buffer ready to consume.
diff --git a/Engine/Plugins/Media/PixelCapture/Source/PixelCapture/Private/PixelCaptureCapturer.cpp b/Engine/Plugins/Media/PixelCapture/Source/PixelCapture/Private/PixelCaptureCapturer.cpp
--- a/Engine/Plugins/Media/PixelCapture/Source/PixelCapture/Private/PixelCaptureCapturer.cpp
+++ b/Engine/Plugins/Media/PixelCapture/Source/PixelCapture/Private/PixelCaptureCapturer.cpp
@@ -4,6 +4,77 @@
 #include "OutputFrameBuffer.h"
 #include "PixelCapturePrivate.h"
 
+namespace UE::PixelCapture
+{
+	void MarkCPUWorkStart(TSharedPtr<IPixelCaptureOutputFrame> OutputBuffer)
+	{
+		if (!OutputBuffer)
+		{
+			UE_LOG(LogPixelCapture, Warning, TEXT("Unable to update metadata on a null output buffer!"));
+			return;
+		}
+
+		OutputBuffer->Metadata.CaptureProcessCPUStartCycles = FPlatformTime::Cycles64();
+	}
+
+	void MarkCPUWorkEnd(TSharedPtr<IPixelCaptureOutputFrame> OutputBuffer)
+	{
+		if (!OutputBuffer)
+		{
+			UE_LOG(LogPixelCapture, Warning, TEXT("Unable to update metadata on a null output buffer!"));
+			return;
+		}
+
+		const uint64 CyclesNow = FPlatformTime::Cycles64();
+
+		if (OutputBuffer->Metadata.CaptureProcessCPUStartCycles == 0)
+		{
+			OutputBuffer->Metadata.CaptureProcessCPUStartCycles = CyclesNow;
+		}
+
+		OutputBuffer->Metadata.CaptureProcessCPUEndCycles = CyclesNow;
+		OutputBuffer->Metadata.CaptureProcessGPUEnqueueStartCycles = CyclesNow;
+	}
+
+	void MarkGPUWorkStart(TSharedPtr<IPixelCaptureOutputFrame> OutputBuffer)
+	{
+		if (!OutputBuffer)
+		{
+			UE_LOG(LogPixelCapture, Warning, TEXT("Unable to update metadata on a null output buffer!"));
+			return;
+		}
+
+		const uint64 CyclesNow = FPlatformTime::Cycles64();
+
+		if (OutputBuffer->Metadata.CaptureProcessGPUEnqueueStartCycles == 0)
+		{
+			OutputBuffer->Metadata.CaptureProcessGPUEnqueueStartCycles = CyclesNow;
+		}
+
+		OutputBuffer->Metadata.CaptureProcessGPUEnqueueEndCycles = CyclesNow;
+		OutputBuffer->Metadata.CaptureProcessGPUStartCycles = CyclesNow;
+	}
+
+	void MarkGPUWorkEnd(TSharedPtr<IPixelCaptureOutputFrame> OutputBuffer)
+	{
+		if (!OutputBuffer)
+		{
+			UE_LOG(LogPixelCapture, Warning, TEXT("Unable to update metadata on a null output buffer!"));
+			return;
+		}
+
+		const uint64 CyclesNow = FPlatformTime::Cycles64();
+
+		if (OutputBuffer->Metadata.CaptureProcessGPUStartCycles == 0)
+		{
+			OutputBuffer->Metadata.CaptureProcessGPUStartCycles = CyclesNow;
+		}
+
+		OutputBuffer->Metadata.CaptureProcessGPUEndCycles = CyclesNow;
+		OutputBuffer->Metadata.CaptureProcessPostGPUStartCycles = CyclesNow;
+	}
+} // namespace UE::PixelCapture
+
 // defined here so we can delete the FOutputFrameBuffer
 FPixelCaptureCapturer::FPixelCaptureCapturer() = default;
 FPixelCaptureCapturer::~FPixelCaptureCapturer() = default;
@@ -24,7 +95,6 @@ void FPixelCaptureCapturer::Capture(const IPixelCaptureInputFrame& InputFrame)
 		return;
 	}
 
-	bBusy = true;
 	const int32 InputWidth = InputFrame.GetWidth();
 	const int32 InputHeight = InputFrame.GetHeight();
 
@@ -35,21 +105,16 @@ void FPixelCaptureCapturer::Capture(const IPixelCaptureInputFrame& InputFrame)
 
 	checkf(InputWidth == ExpectedInputWidth && InputHeight == ExpectedInputHeight, TEXT("Capturer input resolution changes are not supported"));
 
-	CurrentOutputBuffer = Buffer->LockProduceBuffer();
-	if (!CurrentOutputBuffer)
+	TSharedPtr<IPixelCaptureOutputFrame> OutputBuffer = Buffer->LockProduceBuffer();
+	if (!OutputBuffer)
 	{
 		UE_LOG(LogPixelCapture, Error, TEXT("Failed to obtain a produce buffer."));
 		return;
 	}
 
-	InitMetadata(InputFrame.Metadata.Copy());
-	StartTime = rtc::TimeMillis();
-	CPUStartTime = 0;
-	GPUEnqueueTime = 0;
-	GPUStartTime = 0;
+	InitMetadata(InputFrame.Metadata.Copy(), OutputBuffer);
 
-	// Todo (Luke.Bermingham) - Converting output buffer to raw ptr here seems error prone considering how Buffer->LockProduceBuffer() works.
-	BeginProcess(InputFrame, CurrentOutputBuffer.Get());
+	BeginProcess(InputFrame, OutputBuffer);
 }
 
 void FPixelCaptureCapturer::Initialize(int32 InputWidth, int32 InputHeight)
@@ -63,77 +128,71 @@ void FPixelCaptureCapturer::Initialize(int32 InputWidth, int32 InputHeight)
 	bInitialized = true;
 }
 
-void FPixelCaptureCapturer::MarkCPUWorkStart()
+void FPixelCaptureCapturer::InitMetadata(FPixelCaptureFrameMetadata Metadata, TSharedPtr<IPixelCaptureOutputFrame> OutputBuffer)
 {
-	if (CPUStartTime != 0)
+	if (!OutputBuffer)
 	{
-		MarkCPUWorkEnd();
+		UE_LOG(LogPixelCapture, Warning, TEXT("Unable to update metadata on a null output buffer!"));
+		return;
 	}
-	CPUStartTime = rtc::TimeMillis();
-}
 
-void FPixelCaptureCapturer::MarkCPUWorkEnd()
-{
-	check(CurrentOutputBuffer != nullptr);
-	CurrentOutputBuffer->Metadata.CaptureProcessCPUTime += rtc::TimeMillis() - CPUStartTime;
-	CPUStartTime = 0;
+	Metadata.Id = FrameId.Increment();
+	Metadata.ProcessName = GetCapturerName();
 
-	GPUEnqueueTime = rtc::TimeMillis();
+	PRAGMA_DISABLE_DEPRECATION_WARNINGS
+	OutputBuffer->Metadata = Metadata;
+	PRAGMA_ENABLE_DEPRECATION_WARNINGS
+
+	// Reset all the capture timestamps
+	OutputBuffer->Metadata.CaptureStartCyles = 0;
+	OutputBuffer->Metadata.CaptureEndCyles = 0;
+	OutputBuffer->Metadata.CaptureProcessCPUStartCycles = 0;
+	OutputBuffer->Metadata.CaptureProcessCPUEndCycles = 0;
+	OutputBuffer->Metadata.CaptureProcessGPUEnqueueStartCycles = 0;
+	OutputBuffer->Metadata.CaptureProcessGPUEnqueueEndCycles = 0;
+	OutputBuffer->Metadata.CaptureProcessGPUStartCycles = 0;
+	OutputBuffer->Metadata.CaptureProcessGPUEndCycles = 0;
+	OutputBuffer->Metadata.CaptureProcessPostGPUStartCycles = 0;
+	OutputBuffer->Metadata.CaptureProcessPostGPUEndCycles = 0;
+
+	OutputBuffer->Metadata.CaptureStartCyles = FPlatformTime::Cycles64();
 }
 
-void FPixelCaptureCapturer::MarkGPUWorkStart()
+void FPixelCaptureCapturer::FinalizeMetadata(TSharedPtr<IPixelCaptureOutputFrame> OutputBuffer)
 {
-	if (GPUStartTime != 0)
+	if (!OutputBuffer)
 	{
-		MarkGPUWorkEnd();
+		UE_LOG(LogPixelCapture, Warning, TEXT("Unable to update metadata on a null output buffer!"));
+		return;
 	}
-	GPUStartTime = rtc::TimeMillis();
 
-	CurrentOutputBuffer->Metadata.CaptureProcessGPUDelay += GPUStartTime - GPUEnqueueTime;
-	GPUEnqueueTime = 0;
-}
+	const uint64 CyclesNow = FPlatformTime::Cycles64();
 
-void FPixelCaptureCapturer::MarkGPUWorkEnd()
-{
-	CurrentOutputBuffer->Metadata.CaptureProcessGPUTime += rtc::TimeMillis() - GPUStartTime;
-	GPUStartTime = 0;
-}
+	if (OutputBuffer->Metadata.CaptureProcessPostGPUEndCycles == 0)
+	{
+		OutputBuffer->Metadata.CaptureProcessPostGPUEndCycles = CyclesNow;
+	}
 
-void FPixelCaptureCapturer::InitMetadata(FPixelCaptureFrameMetadata Metadata)
-{
-	Metadata.Id = FrameId.Increment();
-	Metadata.ProcessName = GetCapturerName();
-	Metadata.CaptureTime = 0;
-	Metadata.CaptureProcessCPUTime = 0;
-	Metadata.CaptureProcessGPUTime = 0;
-	CurrentOutputBuffer->Metadata = Metadata;
+	OutputBuffer->Metadata.CaptureEndCyles = CyclesNow;
 }
 
-void FPixelCaptureCapturer::FinalizeMetadata()
+void FPixelCaptureCapturer::EndProcess(TSharedPtr<IPixelCaptureOutputFrame> OutputBuffer)
 {
-	if (CPUStartTime != 0)
+	if (!OutputBuffer)
 	{
-		MarkCPUWorkEnd();
+		UE_LOG(LogPixelCapture, Warning, TEXT("Unable to end process on a null output buffer!"));
+		return;
 	}
 
-	if (GPUStartTime != 0)
+	FinalizeMetadata(OutputBuffer);
+
+	if (Buffer->ReleaseProduceBuffer(OutputBuffer))
 	{
-		MarkGPUWorkEnd();
+		bHasOutput = true;
+		OnComplete.Broadcast();
+	}
+	else
+	{
+		bHasOutput = false;
 	}
-
-	CurrentOutputBuffer->Metadata.CaptureTime += rtc::TimeMillis() - StartTime;
-}
-
-void FPixelCaptureCapturer::EndProcess()
-{
-	checkf(bBusy, TEXT("Capture process EndProcess called but we're not busy. Maybe double called?"));
-
-	FinalizeMetadata();
-
-	CurrentOutputBuffer = nullptr;
-	Buffer->ReleaseProduceBuffer();
-	bBusy = false;
-	bHasOutput = true;
-
-	OnComplete.Broadcast();
 }
diff --git a/Engine/Plugins/Media/PixelCapture/Source/PixelCapture/Private/PixelCaptureCapturerI420.cpp b/Engine/Plugins/Media/PixelCapture/Source/PixelCapture/Private/PixelCaptureCapturerI420.cpp
--- a/Engine/Plugins/Media/PixelCapture/Source/PixelCapture/Private/PixelCaptureCapturerI420.cpp
+++ b/Engine/Plugins/Media/PixelCapture/Source/PixelCapture/Private/PixelCaptureCapturerI420.cpp
@@ -15,17 +15,21 @@ IPixelCaptureOutputFrame* FPixelCaptureCapturerI420::CreateOutputBuffer(int32 In
 	return new FPixelCaptureOutputFrameI420(MakeShared<FPixelCaptureBufferI420>(InputWidth, InputHeight));
 }
 
-void FPixelCaptureCapturerI420::BeginProcess(const IPixelCaptureInputFrame& InputFrame, IPixelCaptureOutputFrame* OutputBuffer)
+void FPixelCaptureCapturerI420::BeginProcess(const IPixelCaptureInputFrame& InputFrame, TSharedPtr<IPixelCaptureOutputFrame> OutputBuffer)
 {
+	SetIsBusy(true);
+
 	checkf(InputFrame.GetType() == StaticCast<int32>(PixelCaptureBufferFormat::FORMAT_I420), TEXT("Incorrect source frame coming into frame capture process, expected FORMAT_I420 frame."));
 
-	MarkCPUWorkStart();
+	UE::PixelCapture::MarkCPUWorkStart(OutputBuffer);
 
 	// just copy the input frame
 	const FPixelCaptureInputFrameI420& SourceFrame = StaticCast<const FPixelCaptureInputFrameI420&>(InputFrame);
-	FPixelCaptureOutputFrameI420* OutputI420Buffer = StaticCast<FPixelCaptureOutputFrameI420*>(OutputBuffer);
+	TSharedPtr<FPixelCaptureOutputFrameI420> OutputI420Buffer = StaticCastSharedPtr<FPixelCaptureOutputFrameI420>(OutputBuffer);
 	OutputI420Buffer->GetI420Buffer()->Copy(*SourceFrame.GetBuffer());
 
-	MarkCPUWorkEnd();
-	EndProcess();
+	UE::PixelCapture::MarkCPUWorkEnd(OutputBuffer);
+
+	EndProcess(OutputBuffer);
+	SetIsBusy(false);
 }
diff --git a/Engine/Plugins/Media/PixelCapture/Source/PixelCapture/Private/PixelCaptureCapturerI420ToRHI.cpp b/Engine/Plugins/Media/PixelCapture/Source/PixelCapture/Private/PixelCaptureCapturerI420ToRHI.cpp
--- a/Engine/Plugins/Media/PixelCapture/Source/PixelCapture/Private/PixelCaptureCapturerI420ToRHI.cpp
+++ b/Engine/Plugins/Media/PixelCapture/Source/PixelCapture/Private/PixelCaptureCapturerI420ToRHI.cpp
@@ -30,7 +30,7 @@ FPixelCaptureCapturerI420ToRHI::~FPixelCaptureCapturerI420ToRHI()
 
 void FPixelCaptureCapturerI420ToRHI::Initialize(int32 InputWidth, int32 InputHeight)
 {
-	Dimensions = {InputWidth, InputHeight};
+	Dimensions = { InputWidth, InputHeight };
 
 	const uint32_t Size = webrtc::CalcBufferSize(webrtc::VideoType::kARGB, Dimensions.X, Dimensions.Y);
 	ARGBBuffer = new uint8_t[Size];
@@ -38,7 +38,7 @@ void FPixelCaptureCapturerI420ToRHI::Initialize(int32 InputWidth, int32 InputHei
 	FPixelCaptureCapturer::Initialize(InputWidth, InputHeight);
 }
 
-IPixelCaptureOutputFrame *FPixelCaptureCapturerI420ToRHI::CreateOutputBuffer(int32 InputWidth, int32 InputHeight)
+IPixelCaptureOutputFrame* FPixelCaptureCapturerI420ToRHI::CreateOutputBuffer(int32 InputWidth, int32 InputHeight)
 {
 	FRHITextureCreateDesc TextureDesc =
 		FRHITextureCreateDesc::Create2D(TEXT("FPixelCaptureCapturerI420ToRHI Texture"), InputWidth, InputHeight, EPixelFormat::PF_B8G8R8A8)
@@ -59,15 +59,17 @@ IPixelCaptureOutputFrame *FPixelCaptureCapturerI420ToRHI::CreateOutputBuffer(int
 	return new FPixelCaptureOutputFrameRHI(RHICreateTexture(TextureDesc));
 }
 
-void FPixelCaptureCapturerI420ToRHI::BeginProcess(const IPixelCaptureInputFrame &InputFrame, IPixelCaptureOutputFrame *OutputBuffer)
+void FPixelCaptureCapturerI420ToRHI::BeginProcess(const IPixelCaptureInputFrame& InputFrame, TSharedPtr<IPixelCaptureOutputFrame> OutputBuffer)
 {
+	SetIsBusy(true);
+
 	if (InputFrame.GetType() == StaticCast<int32>(PixelCaptureBufferFormat::FORMAT_I420))
 	{
-		MarkCPUWorkStart();
+		UE::PixelCapture::MarkCPUWorkStart(OutputBuffer);
 
-		const FPixelCaptureInputFrameI420 &SourceFrame = StaticCast<const FPixelCaptureInputFrameI420 &>(InputFrame);
-		TSharedPtr<FPixelCaptureBufferI420> I420Buffer = SourceFrame.GetBuffer();
-		FPixelCaptureOutputFrameRHI *OutputTexture = StaticCast<FPixelCaptureOutputFrameRHI *>(OutputBuffer);
+		const FPixelCaptureInputFrameI420&		SourceFrame = StaticCast<const FPixelCaptureInputFrameI420&>(InputFrame);
+		TSharedPtr<FPixelCaptureBufferI420>		I420Buffer = SourceFrame.GetBuffer();
+		TSharedPtr<FPixelCaptureOutputFrameRHI> OutputTexture = StaticCastSharedPtr<FPixelCaptureOutputFrameRHI>(OutputBuffer);
 
 		// convert the incoming i420 to argb
 		libyuv::ConvertFromI420(
@@ -80,21 +82,21 @@ void FPixelCaptureCapturerI420ToRHI::BeginProcess(const IPixelCaptureInputFrame
 
 		TWeakPtr<FPixelCaptureCapturerI420ToRHI> WeakSelf = StaticCastSharedRef<FPixelCaptureCapturerI420ToRHI>(AsShared());
 		ENQUEUE_RENDER_COMMAND(I420toRHI)
-		([WeakSelf, OutputTexture](FRHICommandListImmediate &RHICmdList) -> void
-		 {
+		([WeakSelf, OutputTexture](FRHICommandListImmediate& RHICmdList) -> void {
 			if (TSharedPtr<FPixelCaptureCapturerI420ToRHI> Self = WeakSelf.Pin())
 			{
-				Self->MarkGPUWorkStart();
+				UE::PixelCapture::MarkGPUWorkStart(OutputTexture);
 
 				// update our staging texture
 				FUpdateTextureRegion2D Region(0, 0, 0, 0, Self->Dimensions.X, Self->Dimensions.Y);
 				RHICmdList.UpdateTexture2D(OutputTexture->GetFrameTexture(), 0, Region, Self->Dimensions.X * 4, Self->ARGBBuffer);
 
-				Self->MarkGPUWorkEnd();
-				Self->EndProcess();
+				UE::PixelCapture::MarkGPUWorkEnd(OutputTexture);
+				Self->EndProcess(OutputTexture);
+				Self->SetIsBusy(false);
 			} });
 
-		MarkCPUWorkEnd();
+		UE::PixelCapture::MarkCPUWorkEnd(OutputBuffer);
 	}
 	else
 	{
diff --git a/Engine/Plugins/Media/PixelCapture/Source/PixelCapture/Private/PixelCaptureCapturerLayered.cpp b/Engine/Plugins/Media/PixelCapture/Source/PixelCapture/Private/PixelCaptureCapturerLayered.cpp
--- a/Engine/Plugins/Media/PixelCapture/Source/PixelCapture/Private/PixelCaptureCapturerLayered.cpp
+++ b/Engine/Plugins/Media/PixelCapture/Source/PixelCapture/Private/PixelCaptureCapturerLayered.cpp
@@ -38,22 +38,11 @@ void FPixelCaptureCapturerLayered::AddLayer(float Scale)
 
 void FPixelCaptureCapturerLayered::OnCaptureComplete()
 {
-	--PendingLayers;
-	if (PendingLayers == 0)
-	{
-		OnComplete.Broadcast();
-	}
+	OnComplete.Broadcast();
 }
 
 void FPixelCaptureCapturerLayered::Capture(const IPixelCaptureInputFrame& SourceFrame)
 {
-	if (PendingLayers > 0)
-	{
-		// still capturing previous frame. drop new one.
-		UE_LOG(LogPixelCapture, Warning, TEXT("Layered capture still waiting on %d pending layers."), StaticCast<int>(PendingLayers));
-		return;
-	}
-
 	// work on a temp list so we dont over lock
 	TArray<TSharedPtr<FPixelCaptureCapturer>> TempLayerCapturers;
 	{
@@ -75,9 +64,6 @@ void FPixelCaptureCapturerLayered::Capture(const IPixelCaptureInputFrame& Source
 		TempLayerCapturers = LayerCapturers;
 	}
 
-	// set this before calling process since the process might immediately complete
-	PendingLayers = TempLayerCapturers.Num();
-
 	// capture the frame for encoder use
 	for (auto& LayerCapturer : TempLayerCapturers)
 	{
diff --git a/Engine/Plugins/Media/PixelCapture/Source/PixelCapture/Private/PixelCaptureCapturerMediaCapture.cpp b/Engine/Plugins/Media/PixelCapture/Source/PixelCapture/Private/PixelCaptureCapturerMediaCapture.cpp
--- a/Engine/Plugins/Media/PixelCapture/Source/PixelCapture/Private/PixelCaptureCapturerMediaCapture.cpp
+++ b/Engine/Plugins/Media/PixelCapture/Source/PixelCapture/Private/PixelCaptureCapturerMediaCapture.cpp
@@ -17,7 +17,8 @@ void UPixelCaptureMediaCapture::OnRHIResourceCaptured_AnyThread(
 	TSharedPtr<FMediaCaptureUserData, ESPMode::ThreadSafe> UserData,
 	FTextureRHIRef										   Texture)
 {
-	if (OutputFrame == nullptr)
+	TSharedPtr<IPixelCaptureOutputFrame> OutputFrame;
+	if (!OutputFrames.TryPopFirst(OutputFrame) || !OutputFrame)
 	{
 		UE_LOGFMT(LogPixelCapture, Warning, "UPixelCaptureMediaCapture::OnRHIResourceCaptured_AnyThread: No output frame set!");
 		return;
@@ -25,9 +26,9 @@ void UPixelCaptureMediaCapture::OnRHIResourceCaptured_AnyThread(
 
 	check(Format == PixelCaptureBufferFormat::FORMAT_RHI);
 
-	static_cast<FPixelCaptureOutputFrameRHI*>(OutputFrame)->SetFrameTexture(Texture);
+	StaticCastSharedPtr<FPixelCaptureOutputFrameRHI>(OutputFrame)->SetFrameTexture(Texture);
 
-	OnCaptureComplete.Broadcast();
+	OnCaptureComplete.Broadcast(OutputFrame);
 }
 
 void UPixelCaptureMediaCapture::OnFrameCaptured_AnyThread(
@@ -35,7 +36,8 @@ void UPixelCaptureMediaCapture::OnFrameCaptured_AnyThread(
 	TSharedPtr<FMediaCaptureUserData, ESPMode::ThreadSafe> UserData,
 	const FMediaCaptureResourceData&					   ResourceData)
 {
-	if (OutputFrame == nullptr)
+	TSharedPtr<IPixelCaptureOutputFrame> OutputFrame;
+	if (!OutputFrames.TryPopFirst(OutputFrame) || !OutputFrame)
 	{
 		UE_LOGFMT(LogPixelCapture, Warning, "UPixelCaptureMediaCapture::OnFrameCaptured_AnyThread: No output frame set!");
 		return;
@@ -43,7 +45,6 @@ void UPixelCaptureMediaCapture::OnFrameCaptured_AnyThread(
 
 	check(Format == PixelCaptureBufferFormat::FORMAT_I420);
 
-	FPixelCaptureOutputFrameI420*		I420Frame = static_cast<FPixelCaptureOutputFrameI420*>(OutputFrame);
 	TSharedPtr<FPixelCaptureBufferI420> I420Buffer = MakeShared<FPixelCaptureBufferI420>(ResourceData.Width, ResourceData.Height);
 	libyuv::ARGBToI420(
 		static_cast<uint8*>(ResourceData.Buffer),
@@ -57,9 +58,9 @@ void UPixelCaptureMediaCapture::OnFrameCaptured_AnyThread(
 		I420Buffer->GetWidth(),
 		I420Buffer->GetHeight());
 
-	I420Frame->SetI420Buffer(I420Buffer);
+	StaticCastSharedPtr<FPixelCaptureOutputFrameI420>(OutputFrame)->SetI420Buffer(I420Buffer);
 
-	OnCaptureComplete.Broadcast();
+	OnCaptureComplete.Broadcast(OutputFrame);
 }
 
 void UPixelCaptureMediaCapture::OnCustomCapture_RenderingThread(
@@ -255,7 +256,7 @@ void FPixelCaptureCapturerMediaCapture::InitializeMediaCapture()
 	CaptureOptions.bAutostopOnCapture = true;
 	// This value needs to be >= 1 in order to be a valid configuration
 	// but it isn't actually used with TryCaptureImmediate_RenderThread
-	CaptureOptions.NumberOfFramesToCapture = 1; 
+	CaptureOptions.NumberOfFramesToCapture = 1;
 
 	FRHICaptureResourceDescription ResourceDescription;
 	ResourceDescription.PixelFormat = EPixelFormat::PF_B8G8R8A8;
@@ -287,47 +288,55 @@ IPixelCaptureOutputFrame* FPixelCaptureCapturerMediaCapture::CreateOutputBuffer(
 	}
 }
 
-void FPixelCaptureCapturerMediaCapture::BeginProcess(const IPixelCaptureInputFrame& InputFrame, IPixelCaptureOutputFrame* OutputBuffer)
+void FPixelCaptureCapturerMediaCapture::BeginProcess(const IPixelCaptureInputFrame& InputFrame, TSharedPtr<IPixelCaptureOutputFrame> OutputBuffer)
 {
 	if (!bMediaCaptureInitialized)
 	{
 		// Early out as media capture is still initializing itself. We'll capture a later frame
-		EndProcess();
+		EndProcess(OutputBuffer);
 		return;
 	}
 
 	checkf(InputFrame.GetType() == StaticCast<int32>(PixelCaptureBufferFormat::FORMAT_RHI), TEXT("Incorrect source frame coming into frame capture process."));
 	const FPixelCaptureInputFrameRHI& SourceFrame = StaticCast<const FPixelCaptureInputFrameRHI&>(InputFrame);
 
-	if (Format == PixelCaptureBufferFormat::FORMAT_RHI)
-	{
-		// If the source texture already matches the dimensions and pixelformat we're looking for
-		// we can just assign the input to the output and early out
-		const FRHITextureDesc& SourceDesc = SourceFrame.FrameTexture->GetDesc();
-		if (SourceDesc.Extent == MediaOutput->GetRequestedSize()
-			&& SourceDesc.Format == EPixelFormat::PF_B8G8R8A8
-#if PLATFORM_MAC
-			// Mac output textures have to have the CPUReadback flag, so if the input doesn't have it we can't do no copy process
-			&& EnumHasAnyFlags(SourceDesc.Flags, TexCreate_CPUReadback)
-#endif
-		)
-		{
-			static_cast<FPixelCaptureOutputFrameRHI*>(OutputBuffer)->SetFrameTexture(SourceFrame.FrameTexture);
-			EndProcess();
-			return;
-		}
-	}
+	UE::PixelCapture::MarkCPUWorkStart(OutputBuffer);
 
-	MediaCapture->SetOutputFrame(OutputBuffer);
+	// We need to ensure the output frame is valid before calling TryCaptureImmediate_RenderThread
+	MediaCapture->AddOutputFrame(OutputBuffer);
 
 	FRDGBuilder GraphBuilder(FRHICommandListImmediate::Get());
-	bool		bPassesAdded = MediaCapture->TryCaptureImmediate_RenderThread(GraphBuilder, SourceFrame.FrameTexture);
+
+	// clang-format off
+	GraphBuilder.AddPass(
+		RDG_EVENT_NAME("UPixelCaptureMediaCapture MarkGPUWorkStart"), 
+		ERDGPassFlags::NeverCull,								
+		[this, OutputBuffer](FRDGAsyncTask, FRHICommandList& RHICmdList) {
+			UE::PixelCapture::MarkGPUWorkStart(OutputBuffer);
+		});
+	// clang-format on
+
+	bool bPassesAdded = MediaCapture->TryCaptureImmediate_RenderThread(GraphBuilder, SourceFrame.FrameTexture);
+
+	// clang-format off
+	GraphBuilder.AddPass(
+		RDG_EVENT_NAME("UPixelCaptureMediaCapture MarkGPUWorkEnd"), 
+		ERDGPassFlags::NeverCull,								
+		[this, OutputBuffer](FRDGAsyncTask, FRHICommandList& RHICmdList) {
+			UE::PixelCapture::MarkGPUWorkEnd(OutputBuffer);
+		});
+	// clang-format on
+
+	UE::PixelCapture::MarkCPUWorkEnd(OutputBuffer);
+
 	// Even if no passes are added, we still need to call Execute
 	GraphBuilder.Execute();
 
 	if (!bPassesAdded)
 	{
+		// TryCaptureImmediate_RenderThread didn't add any passes so remove the latest enqueued output frame from the MediaCapture
+		MediaCapture->RemoveLastOutputFrame();
 		// RDG graph had no passes so we can manually call EndProcess()
-		EndProcess();
+		EndProcess(OutputBuffer);
 	}
 }
\ No newline at end of file
diff --git a/Engine/Plugins/Media/PixelCapture/Source/PixelCapture/Private/PixelCaptureCapturerMultiFormat.cpp b/Engine/Plugins/Media/PixelCapture/Source/PixelCapture/Private/PixelCaptureCapturerMultiFormat.cpp
--- a/Engine/Plugins/Media/PixelCapture/Source/PixelCapture/Private/PixelCaptureCapturerMultiFormat.cpp
+++ b/Engine/Plugins/Media/PixelCapture/Source/PixelCapture/Private/PixelCaptureCapturerMultiFormat.cpp
@@ -35,12 +35,6 @@ int32 FPixelCaptureCapturerMultiFormat::GetHeight(int LayerIndex) const
 
 void FPixelCaptureCapturerMultiFormat::Capture(const IPixelCaptureInputFrame& SourceFrame)
 {
-	if (PendingFormats > 0)
-	{
-		// still capturing previous frame. drop new one.
-		return;
-	}
-
 	// cache each layer size
 	if (LayerSizes.IsEmpty())
 	{
@@ -69,7 +63,6 @@ void FPixelCaptureCapturerMultiFormat::Capture(const IPixelCaptureInputFrame& So
 			AllFormatCapturers.Add(FormatCapturer.Value);
 		}
 		// start all the format captures
-		PendingFormats = AllFormatCapturers.Num();
 		for (auto& Capturer : AllFormatCapturers)
 		{
 			Capturer->Capture(SourceFrame);
@@ -77,11 +70,7 @@ void FPixelCaptureCapturerMultiFormat::Capture(const IPixelCaptureInputFrame& So
 	}
 	else
 	{
-		// if we have no capturers just signal this capture as complete
-		if (PendingFormats == 0)
-		{
-			OnComplete.Broadcast();
-		}
+		OnComplete.Broadcast();
 	}
 }
 
@@ -140,12 +129,8 @@ void FPixelCaptureCapturerMultiFormat::OnDisconnected()
 
 void FPixelCaptureCapturerMultiFormat::OnCaptureFormatComplete(int32 Format)
 {
-	--PendingFormats;
 	CheckFormatEvent(Format); // checks if anything is waiting on this format
-	if (PendingFormats == 0)
-	{
-		OnComplete.Broadcast();
-	}
+	OnComplete.Broadcast();
 }
 
 FEvent* FPixelCaptureCapturerMultiFormat::GetEventForFormat(int32 Format)
diff --git a/Engine/Plugins/Media/PixelCapture/Source/PixelCapture/Private/PixelCaptureCapturerNV12ToRHI.cpp b/Engine/Plugins/Media/PixelCapture/Source/PixelCapture/Private/PixelCaptureCapturerNV12ToRHI.cpp
--- a/Engine/Plugins/Media/PixelCapture/Source/PixelCapture/Private/PixelCaptureCapturerNV12ToRHI.cpp
+++ b/Engine/Plugins/Media/PixelCapture/Source/PixelCapture/Private/PixelCaptureCapturerNV12ToRHI.cpp
@@ -58,15 +58,17 @@ IPixelCaptureOutputFrame* FPixelCaptureCapturerNV12ToRHI::CreateOutputBuffer(int
 	return new FPixelCaptureOutputFrameRHI(RHICreateTexture(TextureDesc), { InputWidth, InputHeight, EPixelFormat::PF_NV12 });
 }
 
-void FPixelCaptureCapturerNV12ToRHI::BeginProcess(const IPixelCaptureInputFrame& InputFrame, IPixelCaptureOutputFrame* OutputBuffer)
+void FPixelCaptureCapturerNV12ToRHI::BeginProcess(const IPixelCaptureInputFrame& InputFrame, TSharedPtr<IPixelCaptureOutputFrame> OutputBuffer)
 {
+	SetIsBusy(true);
+
 	if (InputFrame.GetType() == StaticCast<int32>(PixelCaptureBufferFormat::FORMAT_NV12))
 	{
-		MarkCPUWorkStart();
+		UE::PixelCapture::MarkCPUWorkStart(OutputBuffer);
 
 		const FPixelCaptureInputFrameNV12& SourceFrame = StaticCast<const FPixelCaptureInputFrameNV12&>(InputFrame);
 		TSharedPtr<FPixelCaptureBufferNV12> NV12Buffer = SourceFrame.GetBuffer();
-		FPixelCaptureOutputFrameRHI* OutputTexture = StaticCast<FPixelCaptureOutputFrameRHI*>(OutputBuffer);
+		TSharedPtr<FPixelCaptureOutputFrameRHI> OutputTexture = StaticCastSharedPtr<FPixelCaptureOutputFrameRHI>(OutputBuffer);
 
 		checkf(OutputTexture->GetFrameTexture()->GetDesc().Format == EPixelFormat::PF_R8, TEXT("Texture Format is incorrect for NV12 storage"))
 
@@ -74,18 +76,20 @@ void FPixelCaptureCapturerNV12ToRHI::BeginProcess(const IPixelCaptureInputFrame&
 		ENQUEUE_RENDER_COMMAND(PixelCaptureNV12ToRHI)([WeakSelf, OutputTexture, NV12Buffer](FRHICommandListImmediate& RHICmdList) {
 			if (TSharedPtr<FPixelCaptureCapturerNV12ToRHI> Self = WeakSelf.Pin())
 			{
-				Self->MarkGPUWorkStart();
+				UE::PixelCapture::MarkGPUWorkStart(OutputTexture);
 
 				// update our staging texture
 				FUpdateTextureRegion2D Region(0, 0, 0, 0, Self->Dimensions.X, Self->Dimensions.Y * 1.5);
 				GDynamicRHI->RHIUpdateTexture2D(RHICmdList, OutputTexture->GetFrameTexture(), 0, Region, Self->Dimensions.X, NV12Buffer->GetData());
 
-				Self->MarkGPUWorkEnd();
-				Self->EndProcess();
+				UE::PixelCapture::MarkGPUWorkEnd(OutputTexture);
+
+				Self->EndProcess(OutputTexture);
+				Self->SetIsBusy(false);
 			}
 		});
 
-		MarkCPUWorkEnd();
+		UE::PixelCapture::MarkCPUWorkEnd(OutputBuffer);
 	}
 	else
 	{
diff --git a/Engine/Plugins/Media/PixelCapture/Source/PixelCapture/Private/PixelCaptureCapturerRHI.cpp b/Engine/Plugins/Media/PixelCapture/Source/PixelCapture/Private/PixelCaptureCapturerRHI.cpp
--- a/Engine/Plugins/Media/PixelCapture/Source/PixelCapture/Private/PixelCaptureCapturerRHI.cpp
+++ b/Engine/Plugins/Media/PixelCapture/Source/PixelCapture/Private/PixelCaptureCapturerRHI.cpp
@@ -16,7 +16,6 @@ TSharedPtr<FPixelCaptureCapturerRHI> FPixelCaptureCapturerRHI::Create(float InSc
 FPixelCaptureCapturerRHI::FPixelCaptureCapturerRHI(float InScale)
 	: Scale(InScale)
 {
-	Fence = GDynamicRHI->RHICreateGPUFence(TEXT("FPixelCaptureCapturerRHI Fence"));
 	RHIType = RHIGetInterfaceType();
 }
 
@@ -40,7 +39,7 @@ IPixelCaptureOutputFrame* FPixelCaptureCapturerRHI::CreateOutputBuffer(int32 Inp
 			.SetInitialState(ERHIAccess::Present);
 	}
 
-	TextureDesc.DetermineInititialState();		
+	TextureDesc.DetermineInititialState();
 
 	if (RHIGetInterfaceType() == ERHIInterfaceType::Vulkan)
 	{
@@ -54,32 +53,33 @@ IPixelCaptureOutputFrame* FPixelCaptureCapturerRHI::CreateOutputBuffer(int32 Inp
 	return new FPixelCaptureOutputFrameRHI(RHICreateTexture(TextureDesc));
 }
 
-void FPixelCaptureCapturerRHI::BeginProcess(const IPixelCaptureInputFrame& InputFrame, IPixelCaptureOutputFrame* OutputBuffer)
+void FPixelCaptureCapturerRHI::BeginProcess(const IPixelCaptureInputFrame& InputFrame, TSharedPtr<IPixelCaptureOutputFrame> OutputBuffer)
 {
 	checkf(InputFrame.GetType() == StaticCast<int32>(PixelCaptureBufferFormat::FORMAT_RHI), TEXT("Incorrect source frame coming into frame capture process."));
 
-	MarkCPUWorkStart();
+	UE::PixelCapture::MarkCPUWorkStart(OutputBuffer);
 
 	FRHICommandListImmediate& RHICmdList = FRHICommandListExecutor::GetImmediateCommandList();
-	RHICmdList.EnqueueLambda([this](FRHICommandListImmediate&) { MarkGPUWorkStart(); });
+	RHICmdList.EnqueueLambda([this, OutputBuffer](FRHICommandListImmediate&) { UE::PixelCapture::MarkGPUWorkStart(OutputBuffer); });
 
-	const FPixelCaptureInputFrameRHI& RHISourceFrame = StaticCast<const FPixelCaptureInputFrameRHI&>(InputFrame);
-	FPixelCaptureOutputFrameRHI* OutputH264Buffer = StaticCast<FPixelCaptureOutputFrameRHI*>(OutputBuffer);
+	const FPixelCaptureInputFrameRHI&		RHISourceFrame = StaticCast<const FPixelCaptureInputFrameRHI&>(InputFrame);
+	TSharedPtr<FPixelCaptureOutputFrameRHI> OutputH264Buffer = StaticCastSharedPtr<FPixelCaptureOutputFrameRHI>(OutputBuffer);
+	FGPUFenceRHIRef							Fence = GDynamicRHI->RHICreateGPUFence(TEXT("FPixelCaptureCapturerRHI Fence"));
 	CopyTexture(RHICmdList, RHISourceFrame.FrameTexture, OutputH264Buffer->GetFrameTexture(), Fence);
 
-	MarkCPUWorkEnd();
+	UE::PixelCapture::MarkCPUWorkEnd(OutputBuffer);
 
-	// by adding this shared ref to the rhi lambda we can ensure that 'this' will not be destroyed
+	// by adding this shared ref to the async lambda we can ensure that 'this' will not be destroyed
 	// until after the rhi thread is done with it, so all the commands will still have valid references.
 	TSharedRef<FPixelCaptureCapturerRHI> ThisRHIRef = StaticCastSharedRef<FPixelCaptureCapturerRHI>(AsShared());
-	RHICmdList.EnqueueLambda([ThisRHIRef](FRHICommandListImmediate&) { ThisRHIRef->CheckComplete(); });
+	AsyncTask(ENamedThreads::AnyBackgroundHiPriTask, [ThisRHIRef, OutputBuffer, Fence]() { ThisRHIRef->CheckComplete(OutputBuffer, Fence); });
 }
 
-void FPixelCaptureCapturerRHI::CheckComplete()
+void FPixelCaptureCapturerRHI::CheckComplete(TSharedPtr<IPixelCaptureOutputFrame> OutputBuffer, FGPUFenceRHIRef Fence)
 {
 	// TODO: We should move to proper event driven fences once they're implemented. Both DX12 and Vulkan APIs support them, they just haven't been added
 	// to their respective RHIs. DX11_3 supports it but for compatability reasons we can't upgrade from DX11_2.
-	
+
 	// in lieu of a proper callback we need to capture a thread to poll the fence
 	// so we know as quickly as possible when we can readback.
 
@@ -88,18 +88,18 @@ void FPixelCaptureCapturerRHI::CheckComplete()
 	if (!Fence->Poll())
 	{
 		TSharedRef<FPixelCaptureCapturerRHI> ThisRHIRef = StaticCastSharedRef<FPixelCaptureCapturerRHI>(AsShared());
-		AsyncTask(ENamedThreads::AnyBackgroundHiPriTask, [ThisRHIRef]() { ThisRHIRef->CheckComplete(); });
+		AsyncTask(ENamedThreads::AnyBackgroundHiPriTask, [ThisRHIRef, OutputBuffer, Fence]() { ThisRHIRef->CheckComplete(OutputBuffer, Fence); });
 	}
 	else
 	{
-		OnRHIStageComplete();
+		Fence->Clear();
+		OnRHIStageComplete(OutputBuffer);
 	}
 }
 
-void FPixelCaptureCapturerRHI::OnRHIStageComplete()
+void FPixelCaptureCapturerRHI::OnRHIStageComplete(TSharedPtr<IPixelCaptureOutputFrame> OutputBuffer)
 {
-	checkf(Fence->Poll(), TEXT("Fence was not set. Backbuffer copy may not have completed."));
-	Fence->Clear();
-	MarkGPUWorkEnd();
-	EndProcess();
+	UE::PixelCapture::MarkGPUWorkEnd(OutputBuffer);
+
+	EndProcess(OutputBuffer);
 }
diff --git a/Engine/Plugins/Media/PixelCapture/Source/PixelCapture/Private/PixelCaptureCapturerRHINoCopy.cpp b/Engine/Plugins/Media/PixelCapture/Source/PixelCapture/Private/PixelCaptureCapturerRHINoCopy.cpp
--- a/Engine/Plugins/Media/PixelCapture/Source/PixelCapture/Private/PixelCaptureCapturerRHINoCopy.cpp
+++ b/Engine/Plugins/Media/PixelCapture/Source/PixelCapture/Private/PixelCaptureCapturerRHINoCopy.cpp
@@ -23,17 +23,17 @@ IPixelCaptureOutputFrame* FPixelCaptureCapturerRHINoCopy::CreateOutputBuffer(int
 	return new FPixelCaptureOutputFrameRHI(nullptr);
 }
 
-void FPixelCaptureCapturerRHINoCopy::BeginProcess(const IPixelCaptureInputFrame& InputFrame, IPixelCaptureOutputFrame* OutputBuffer)
+void FPixelCaptureCapturerRHINoCopy::BeginProcess(const IPixelCaptureInputFrame& InputFrame, TSharedPtr<IPixelCaptureOutputFrame> OutputBuffer)
 {
 	checkf(InputFrame.GetType() == StaticCast<int32>(PixelCaptureBufferFormat::FORMAT_RHI), TEXT("Incorrect source frame coming into frame capture process."));
 
-	MarkCPUWorkStart();
+	UE::PixelCapture::MarkCPUWorkStart(OutputBuffer);
 
 	const FPixelCaptureInputFrameRHI& RHISourceFrame = StaticCast<const FPixelCaptureInputFrameRHI&>(InputFrame);
-	FPixelCaptureOutputFrameRHI* OutputH264Buffer = StaticCast<FPixelCaptureOutputFrameRHI*>(OutputBuffer);
+	TSharedPtr<FPixelCaptureOutputFrameRHI> OutputH264Buffer = StaticCastSharedPtr<FPixelCaptureOutputFrameRHI>(OutputBuffer);
 	OutputH264Buffer->SetFrameTexture(RHISourceFrame.FrameTexture);
 
-	MarkCPUWorkEnd();
+	UE::PixelCapture::MarkCPUWorkEnd(OutputBuffer);
 
-	EndProcess();
+	EndProcess(OutputBuffer);
 }
\ No newline at end of file
diff --git a/Engine/Plugins/Media/PixelCapture/Source/PixelCapture/Private/PixelCaptureCapturerRHIRDG.cpp b/Engine/Plugins/Media/PixelCapture/Source/PixelCapture/Private/PixelCaptureCapturerRHIRDG.cpp
--- a/Engine/Plugins/Media/PixelCapture/Source/PixelCapture/Private/PixelCaptureCapturerRHIRDG.cpp
+++ b/Engine/Plugins/Media/PixelCapture/Source/PixelCapture/Private/PixelCaptureCapturerRHIRDG.cpp
@@ -52,29 +52,32 @@ IPixelCaptureOutputFrame* FPixelCaptureCapturerRHIRDG::CreateOutputBuffer(int32
 	return new FPixelCaptureOutputFrameRHI(RHICreateTexture(TextureDesc));
 }
 
-void FPixelCaptureCapturerRHIRDG::BeginProcess(const IPixelCaptureInputFrame& InputFrame, IPixelCaptureOutputFrame* OutputBuffer)
+void FPixelCaptureCapturerRHIRDG::BeginProcess(const IPixelCaptureInputFrame& InputFrame, TSharedPtr<IPixelCaptureOutputFrame> OutputBuffer)
 {
+	SetIsBusy(true);
+
 	checkf(InputFrame.GetType() == StaticCast<int32>(PixelCaptureBufferFormat::FORMAT_RHI), TEXT("Incorrect source frame coming into frame capture process."));
 
-	MarkCPUWorkStart();
-	MarkCPUWorkEnd();
+	UE::PixelCapture::MarkCPUWorkStart(OutputBuffer);
+	UE::PixelCapture::MarkCPUWorkEnd(OutputBuffer);
 
 	FRHICommandListImmediate& RHICmdList = FRHICommandListExecutor::GetImmediateCommandList();
 
-	MarkGPUWorkStart();
+	UE::PixelCapture::MarkGPUWorkStart(OutputBuffer);
 	const FPixelCaptureInputFrameRHI& RHISourceFrame = StaticCast<const FPixelCaptureInputFrameRHI&>(InputFrame);
-	FPixelCaptureOutputFrameRHI* OutputH264Buffer = StaticCast<FPixelCaptureOutputFrameRHI*>(OutputBuffer);
+	TSharedPtr<FPixelCaptureOutputFrameRHI> OutputH264Buffer = StaticCastSharedPtr<FPixelCaptureOutputFrameRHI>(OutputBuffer);
 	CopyTextureRDG(RHICmdList, RHISourceFrame.FrameTexture, OutputH264Buffer->GetFrameTexture());
 
-	MarkGPUWorkEnd();
-	OnRHIStageComplete();
+	UE::PixelCapture::MarkGPUWorkEnd(OutputBuffer);
+	OnRHIStageComplete(OutputBuffer);
 }
 
 void FPixelCaptureCapturerRHIRDG::CheckComplete()
 {
 }
 
-void FPixelCaptureCapturerRHIRDG::OnRHIStageComplete()
+void FPixelCaptureCapturerRHIRDG::OnRHIStageComplete(TSharedPtr<IPixelCaptureOutputFrame> OutputBuffer)
 {
-	EndProcess();
+	EndProcess(OutputBuffer);
+	SetIsBusy(false);
 }
diff --git a/Engine/Plugins/Media/PixelCapture/Source/PixelCapture/Private/PixelCaptureCapturerRHIToI420CPU.cpp b/Engine/Plugins/Media/PixelCapture/Source/PixelCapture/Private/PixelCaptureCapturerRHIToI420CPU.cpp
--- a/Engine/Plugins/Media/PixelCapture/Source/PixelCapture/Private/PixelCaptureCapturerRHIToI420CPU.cpp
+++ b/Engine/Plugins/Media/PixelCapture/Source/PixelCapture/Private/PixelCaptureCapturerRHIToI420CPU.cpp
@@ -62,27 +62,29 @@ void FPixelCaptureCapturerRHIToI420CPU::Initialize(int32 InputWidth, int32 Input
 	FPixelCaptureCapturer::Initialize(InputWidth, InputHeight);
 }
 
-void FPixelCaptureCapturerRHIToI420CPU::BeginProcess(const IPixelCaptureInputFrame& InputFrame, IPixelCaptureOutputFrame* OutputBuffer)
+void FPixelCaptureCapturerRHIToI420CPU::BeginProcess(const IPixelCaptureInputFrame& InputFrame, TSharedPtr<IPixelCaptureOutputFrame> OutputBuffer)
 {
+	SetIsBusy(true);
+
 	checkf(InputFrame.GetType() == StaticCast<int32>(PixelCaptureBufferFormat::FORMAT_RHI), TEXT("Incorrect source frame coming into frame capture process."));
 
-	MarkCPUWorkStart();
+	UE::PixelCapture::MarkCPUWorkStart(OutputBuffer);
 
 	const FPixelCaptureInputFrameRHI& RHISourceFrame = StaticCast<const FPixelCaptureInputFrameRHI&>(InputFrame);
-	FTextureRHIRef SourceTexture = RHISourceFrame.FrameTexture;
+	FTextureRHIRef					  SourceTexture = RHISourceFrame.FrameTexture;
 
 	FRHICommandListImmediate& RHICmdList = FRHICommandListExecutor::GetImmediateCommandList();
 
-	RHICmdList.EnqueueLambda([this](FRHICommandListImmediate&) { MarkGPUWorkStart(); });
+	RHICmdList.EnqueueLambda([this, OutputBuffer](FRHICommandListImmediate&) { UE::PixelCapture::MarkGPUWorkStart(OutputBuffer); });
 
 	RHICmdList.Transition(FRHITransitionInfo(SourceTexture, ERHIAccess::Unknown, ERHIAccess::CopySrc));
 	RHICmdList.Transition(FRHITransitionInfo(StagingTexture, ERHIAccess::CopySrc, ERHIAccess::CopyDest));
 	CopyTexture(RHICmdList, SourceTexture, StagingTexture, nullptr);
 	RHICmdList.Transition(FRHITransitionInfo(StagingTexture, ERHIAccess::CopyDest, ERHIAccess::CopySrc));
 
-	TextureReader->EnqueueCopy(RHICmdList, StagingTexture, FIntVector(0, 0, 0), 0, FIntVector(StagingTexture->GetSizeXY().X, StagingTexture->GetSizeXY().Y, 0));
+	UE::PixelCapture::MarkCPUWorkEnd(OutputBuffer);
 
-	MarkCPUWorkEnd();
+	TextureReader->EnqueueCopy(RHICmdList, StagingTexture, FIntVector(0, 0, 0), 0, FIntVector(StagingTexture->GetSizeXY().X, StagingTexture->GetSizeXY().Y, 0));
 
 	// by adding this shared ref to the rhi lambda we can ensure that 'this' will not be destroyed
 	// until after the rhi thread is done with it, so all the commands will still have valid references.
@@ -92,7 +94,7 @@ void FPixelCaptureCapturerRHIToI420CPU::BeginProcess(const IPixelCaptureInputFra
 	});
 }
 
-void FPixelCaptureCapturerRHIToI420CPU::CheckComplete(IPixelCaptureOutputFrame* OutputBuffer)
+void FPixelCaptureCapturerRHIToI420CPU::CheckComplete(TSharedPtr<IPixelCaptureOutputFrame> OutputBuffer)
 {
 	if (!TextureReader->IsReady())
 	{
@@ -106,16 +108,15 @@ void FPixelCaptureCapturerRHIToI420CPU::CheckComplete(IPixelCaptureOutputFrame*
 	}
 }
 
-void FPixelCaptureCapturerRHIToI420CPU::OnRHIStageComplete(IPixelCaptureOutputFrame* OutputBuffer)
+void FPixelCaptureCapturerRHIToI420CPU::OnRHIStageComplete(TSharedPtr<IPixelCaptureOutputFrame> OutputBuffer)
 {
-	MarkGPUWorkEnd();
-	MarkCPUWorkStart();
+	UE::PixelCapture::MarkGPUWorkEnd(OutputBuffer);
 
 	int32 OutRowPitchInPixels = 0;
 	void* PixelData = TextureReader->Lock(OutRowPitchInPixels);
 
-	FPixelCaptureOutputFrameI420* OutputI420Buffer = StaticCast<FPixelCaptureOutputFrameI420*>(OutputBuffer);
-	TSharedPtr<FPixelCaptureBufferI420> I420Buffer = OutputI420Buffer->GetI420Buffer();
+	TSharedPtr<FPixelCaptureOutputFrameI420> OutputI420Buffer = StaticCastSharedPtr<FPixelCaptureOutputFrameI420>(OutputBuffer);
+	TSharedPtr<FPixelCaptureBufferI420>		 I420Buffer = OutputI420Buffer->GetI420Buffer();
 	libyuv::ARGBToI420(
 		static_cast<uint8*>(PixelData),
 		OutRowPitchInPixels * 4,
@@ -130,12 +131,12 @@ void FPixelCaptureCapturerRHIToI420CPU::OnRHIStageComplete(IPixelCaptureOutputFr
 
 	TextureReader->Unlock();
 
-	MarkCPUWorkEnd();
-	EndProcess();
+	EndProcess(OutputBuffer);
+	SetIsBusy(false);
 }
 
 // TODO (Eden.Harris) This has less latency but introduces screen tearing so is disabled. RTCP-7778 to fix the screen tearing without effecting latency.
-#else	
+#else
 
 void FPixelCaptureCapturerRHIToI420CPU::Initialize(int32 InputWidth, int32 InputHeight)
 {
@@ -171,8 +172,10 @@ void FPixelCaptureCapturerRHIToI420CPU::Initialize(int32 InputWidth, int32 Input
 	FPixelCaptureCapturer::Initialize(InputWidth, InputHeight);
 }
 
-void FPixelCaptureCapturerRHIToI420CPU::BeginProcess(const IPixelCaptureInputFrame& InputFrame, IPixelCaptureOutputFrame* OutputBuffer)
+void FPixelCaptureCapturerRHIToI420CPU::BeginProcess(const IPixelCaptureInputFrame& InputFrame, TSharedPtr<IPixelCaptureOutputFrame> OutputBuffer)
 {
+	SetIsBusy(true);
+
 	checkf(InputFrame.GetType() == StaticCast<int32>(PixelCaptureBufferFormat::FORMAT_RHI), TEXT("Incorrect source frame coming into frame capture process."));
 
 	const FPixelCaptureInputFrameRHI& RHISourceFrame = StaticCast<const FPixelCaptureInputFrameRHI&>(InputFrame);
@@ -191,8 +194,8 @@ void FPixelCaptureCapturerRHIToI420CPU::BeginProcess(const IPixelCaptureInputFra
 	// Nullptr fence forces a flush
 	RHICmdList.MapStagingSurface(ReadbackTexture, nullptr, ReadbackPointer, ReadbackWidth, ReadbackHeight);
 
-	FPixelCaptureOutputFrameI420*		OutputI420Buffer = StaticCast<FPixelCaptureOutputFrameI420*>(OutputBuffer);
-	TSharedPtr<FPixelCaptureBufferI420> I420Buffer = OutputI420Buffer->GetI420Buffer();
+	TSharedPtr<FPixelCaptureOutputFrameI420> OutputI420Buffer = StaticCastSharedPtr<FPixelCaptureOutputFrameI420>(OutputBuffer);
+	TSharedPtr<FPixelCaptureBufferI420>		 I420Buffer = OutputI420Buffer->GetI420Buffer();
 	libyuv::ARGBToI420(
 		static_cast<uint8*>(ReadbackPointer),
 		ReadbackWidth * 4,
@@ -207,6 +210,7 @@ void FPixelCaptureCapturerRHIToI420CPU::BeginProcess(const IPixelCaptureInputFra
 
 	RHICmdList.UnmapStagingSurface(ReadbackTexture);
 
-	EndProcess();
+	EndProcess(OutputBuffer);
+	SetIsBusy(false);
 }
 #endif
\ No newline at end of file
diff --git a/Engine/Plugins/Media/PixelCapture/Source/PixelCapture/Private/PixelCaptureCapturerRHIToI420Compute.cpp b/Engine/Plugins/Media/PixelCapture/Source/PixelCapture/Private/PixelCaptureCapturerRHIToI420Compute.cpp
--- a/Engine/Plugins/Media/PixelCapture/Source/PixelCapture/Private/PixelCaptureCapturerRHIToI420Compute.cpp
+++ b/Engine/Plugins/Media/PixelCapture/Source/PixelCapture/Private/PixelCaptureCapturerRHIToI420Compute.cpp
@@ -104,18 +104,20 @@ IPixelCaptureOutputFrame* FPixelCaptureCapturerRHIToI420Compute::CreateOutputBuf
 	return new FPixelCaptureOutputFrameI420(MakeShared<FPixelCaptureBufferI420>(Width, Height));
 }
 
-void FPixelCaptureCapturerRHIToI420Compute::BeginProcess(const IPixelCaptureInputFrame& InputFrame, IPixelCaptureOutputFrame* OutputBuffer)
+void FPixelCaptureCapturerRHIToI420Compute::BeginProcess(const IPixelCaptureInputFrame& InputFrame, TSharedPtr<IPixelCaptureOutputFrame> OutputBuffer)
 {
+	SetIsBusy(true);
+
 	checkf(InputFrame.GetType() == StaticCast<int32>(PixelCaptureBufferFormat::FORMAT_RHI), TEXT("Incorrect source frame coming into frame capture process."));
 
-	MarkCPUWorkStart();
+	UE::PixelCapture::MarkCPUWorkStart(OutputBuffer);
 
 	const FPixelCaptureInputFrameRHI& RHISourceFrame = StaticCast<const FPixelCaptureInputFrameRHI&>(InputFrame);
 	FTextureRHIRef SourceTexture = RHISourceFrame.FrameTexture;
 
 	FRHICommandListImmediate& RHICmdList = FRHICommandListExecutor::GetImmediateCommandList();
 
-	RHICmdList.EnqueueLambda([this](FRHICommandListImmediate&) { MarkGPUWorkStart(); });
+	RHICmdList.EnqueueLambda([this, OutputBuffer](FRHICommandListImmediate&) { UE::PixelCapture::MarkGPUWorkStart(OutputBuffer); });
 
 	FRGBToYUVShaderParameters ShaderParameters;
 	ShaderParameters.SourceTexture = SourceTexture;
@@ -130,7 +132,7 @@ void FPixelCaptureCapturerRHIToI420Compute::BeginProcess(const IPixelCaptureInpu
 	RHICmdList.CopyTexture(TextureU, StagingTextureU, {});
 	RHICmdList.CopyTexture(TextureV, StagingTextureV, {});
 
-	MarkCPUWorkEnd();
+	UE::PixelCapture::MarkCPUWorkEnd(OutputBuffer);
 
 	// by adding this shared ref to the rhi lambda we can ensure that 'this' will not be destroyed
 	// until after the rhi thread is done with it, so all the commands will still have valid references.
@@ -138,21 +140,19 @@ void FPixelCaptureCapturerRHIToI420Compute::BeginProcess(const IPixelCaptureInpu
 	RHICmdList.EnqueueLambda([ThisRHIRef, OutputBuffer](FRHICommandListImmediate&) { ThisRHIRef->OnRHIStageComplete(OutputBuffer); });
 }
 
-void FPixelCaptureCapturerRHIToI420Compute::OnRHIStageComplete(IPixelCaptureOutputFrame* OutputBuffer)
+void FPixelCaptureCapturerRHIToI420Compute::OnRHIStageComplete(TSharedPtr<IPixelCaptureOutputFrame> OutputBuffer)
 {
-	MarkGPUWorkEnd();
-
-	MarkCPUWorkStart();
+	UE::PixelCapture::MarkGPUWorkEnd(OutputBuffer);
 
-	FPixelCaptureOutputFrameI420* OutputI420Buffer = StaticCast<FPixelCaptureOutputFrameI420*>(OutputBuffer);
+	TSharedPtr<FPixelCaptureOutputFrameI420> OutputI420Buffer = StaticCastSharedPtr<FPixelCaptureOutputFrameI420>(OutputBuffer);
 	TSharedPtr<FPixelCaptureBufferI420> I420Buffer = OutputI420Buffer->GetI420Buffer();
 
 	MemCpyStride(I420Buffer->GetMutableDataY(), MappedY, I420Buffer->GetStrideY(), YStride, PlaneYDimensions.Y);
 	MemCpyStride(I420Buffer->GetMutableDataU(), MappedU, I420Buffer->GetStrideUV(), UStride, PlaneUVDimensions.Y);
 	MemCpyStride(I420Buffer->GetMutableDataV(), MappedV, I420Buffer->GetStrideUV(), VStride, PlaneUVDimensions.Y);
 
-	MarkCPUWorkEnd();
-	EndProcess();
+	EndProcess(OutputBuffer);
+	SetIsBusy(false);
 }
 
 void FPixelCaptureCapturerRHIToI420Compute::CleanUp()
diff --git a/Engine/Plugins/Media/PixelCapture/Source/PixelCapture/Private/Tests/PixelCaptureTests.cpp b/Engine/Plugins/Media/PixelCapture/Source/PixelCapture/Private/Tests/PixelCaptureTests.cpp
--- a/Engine/Plugins/Media/PixelCapture/Source/PixelCapture/Private/Tests/PixelCaptureTests.cpp
+++ b/Engine/Plugins/Media/PixelCapture/Source/PixelCapture/Private/Tests/PixelCaptureTests.cpp
@@ -67,24 +67,28 @@ namespace
 			return new FMockOutput();
 		}
 
-		virtual void BeginProcess(const IPixelCaptureInputFrame& InputFrame, IPixelCaptureOutputFrame* OutputBuffer) override
+		virtual void BeginProcess(const IPixelCaptureInputFrame& InputFrame, TSharedPtr<IPixelCaptureOutputFrame> OutputBuffer) override
 		{
+			SetIsBusy(true);
+
 			check(InputFrame.GetType() == StaticCast<int32>(EMockFormats::Input));
 			const FMockInput& Input = StaticCast<const FMockInput&>(InputFrame);
-			FMockOutput* Output = StaticCast<FMockOutput*>(OutputBuffer);
+			TSharedPtr<FMockOutput> Output = StaticCastSharedPtr<FMockOutput>(OutputBuffer);
 			Output->Width = Input.Width * Scale;
 			Output->Height = Input.Height * Scale;
 			Output->MagicData = Input.MagicData;
 
 			if (AutoComplete)
 			{
-				EndProcess();
+				EndProcess(OutputBuffer);
+				SetIsBusy(false);
 			}
 		}
 
-		void MockProcessComplete()
+		void MockProcessComplete(TSharedPtr<IPixelCaptureOutputFrame> OutputBuffer)
 		{
-			EndProcess();
+			EndProcess(OutputBuffer);
+			SetIsBusy(false);
 		}
 	};
 
@@ -144,7 +148,8 @@ namespace UE::PixelCapture
 			TestTrue("Capturer busy after frame input", Capturer.IsBusy() == true);
 			TestTrue("Capturer has no output before completing", Capturer.HasOutput() == false);
 			TestTrue("Capturer still returns null after input but not complete", Capturer.ReadOutput() == nullptr);
-			Capturer.MockProcessComplete();
+			TSharedPtr<IPixelCaptureOutputFrame> OutbutBuffer = MakeShareable<IPixelCaptureOutputFrame>(Capturer.CreateOutputBuffer(1920, 1080));
+			Capturer.MockProcessComplete(OutbutBuffer);
 			TestTrue("Capturer still initialized after process complete", Capturer.IsInitialized() == true);
 			TestTrue("Capturer no longer busy after process complete", Capturer.IsBusy() == false);
 			TestTrue("Capturer has output after completing", Capturer.HasOutput() == true);
diff --git a/Engine/Plugins/Media/PixelCapture/Source/PixelCapture/Public/PixelCaptureCapturer.h b/Engine/Plugins/Media/PixelCapture/Source/PixelCapture/Public/PixelCaptureCapturer.h
--- a/Engine/Plugins/Media/PixelCapture/Source/PixelCapture/Public/PixelCaptureCapturer.h
+++ b/Engine/Plugins/Media/PixelCapture/Source/PixelCapture/Public/PixelCaptureCapturer.h
@@ -9,7 +9,12 @@
 namespace UE::PixelCapture
 {
 	class FOutputFrameBuffer;
-}
+
+	PIXELCAPTURE_API void MarkCPUWorkStart(TSharedPtr<IPixelCaptureOutputFrame> OutputBuffer);
+	PIXELCAPTURE_API void MarkCPUWorkEnd(TSharedPtr<IPixelCaptureOutputFrame> OutputBuffer);
+	PIXELCAPTURE_API void MarkGPUWorkStart(TSharedPtr<IPixelCaptureOutputFrame> OutputBuffer);
+	PIXELCAPTURE_API void MarkGPUWorkEnd(TSharedPtr<IPixelCaptureOutputFrame> OutputBuffer);
+} // namespace UE::PixelCapture
 
 /**
  * The base class for all Capturers in the system.
@@ -65,10 +70,16 @@ public:
 	/**
 	 * Listen on this to be notified when the capture process completes for each input.
 	 */
-	DECLARE_MULTICAST_DELEGATE(FOnComplete);
+	DECLARE_TS_MULTICAST_DELEGATE(FOnComplete);
 	FOnComplete OnComplete;
 
 protected:
+	/**
+	 * Sets whether this process is actively working on capturing frame data.
+	 * @param bInBusy whether this process is busy.
+	 */
+	void SetIsBusy(bool bInBusy) { bBusy = bInBusy; }
+
 	/**
 	 * Gets the human readable name for this capture process. This name will be used in stats
 	 * readouts so the shorter the better.
@@ -96,32 +107,40 @@ protected:
 	 * @param InputFrame The input frame data for the process to begin working on.
 	 * @param OutputBuffer The destination buffer for the process. Is guaranteed to be of the type created in CreateOutputBuffer()
 	 */
-	virtual void BeginProcess(const IPixelCaptureInputFrame& InputFrame, IPixelCaptureOutputFrame* OutputBuffer) = 0;
+	UE_DEPRECATED(5.6, "BeginProcess(const IPixelCaptureInputFrame& InputFrame, IPixelCaptureOutputFrame* OutputBuffer) has been deprecated. Please use BeginProcess(const IPixelCaptureInputFrame& InputFrame, TSharedPtr<IPixelCaptureOutputFrame> OutputBuffer) to ensure memory safety.")
+	void BeginProcess(const IPixelCaptureInputFrame& InputFrame, IPixelCaptureOutputFrame* OutputBuffer) {};
+	virtual void BeginProcess(const IPixelCaptureInputFrame& InputFrame, TSharedPtr<IPixelCaptureOutputFrame> OutputBuffer) = 0;
 
 	/**
 	 * Metadata info (optional). Marks the start of some CPU work. Multiple work sections are valid.
 	 */
-	void MarkCPUWorkStart();
+	UE_DEPRECATED(5.6, "MarkCPUWorkStart(void) has been deprecated. Please use MarkCPUWorkStart(TSharedPtr<IPixelCaptureOutputFrame> OutputBuffer) to ensure the metadata is updated for the correct output buffer.")
+	void MarkCPUWorkStart() {}
 
 	/**
 	 * Metadata info (optional). Marks the end of some CPU work.
 	 */
-	void MarkCPUWorkEnd();
+	UE_DEPRECATED(5.6, "MarkCPUWorkEnd(void) has been deprecated. Please use UE::PixelCapture::MarkCPUWorkEnd(TSharedPtr<IPixelCaptureOutputFrame> OutputBuffer) to ensure the metadata is updated for the correct output buffer.")
+	void MarkCPUWorkEnd() {}
 
 	/**
 	 * Metadata info (optional). Marks the start of some GPU work. Multiple work sections are valid.
 	 */
-	void MarkGPUWorkStart();
+	UE_DEPRECATED(5.6, "MarkGPUWorkStart(void) has been deprecated. Please use UE::PixelCapture::MarkGPUWorkStart(TSharedPtr<IPixelCaptureOutputFrame> OutputBuffer) to ensure the metadata is updated for the correct output buffer.")
+	void MarkGPUWorkStart() {}
 
 	/**
 	 * Metadata info (optional). Marks the end of some GPU work.
 	 */
-	void MarkGPUWorkEnd();
+	UE_DEPRECATED(5.6, "MarkGPUWorkEnd(void) has been deprecated. Please use UE::PixelCapture::MarkGPUWorkEnd(TSharedPtr<IPixelCaptureOutputFrame> OutputBuffer) to ensure the metadata is updated for the correct output buffer.")
+	void MarkGPUWorkEnd() {}
 
 	/**
 	 * Call this to mark the end of processing. Will commit the current write buffer into the read buffer.
 	 */
-	void EndProcess();
+	UE_DEPRECATED(5.6, "EndProcess(void) has been deprecated. Please use EndProcess(TSharedPtr<IPixelCaptureOutputFrame> OutputBuffer) to ensure the metadata is updated for the correct output buffer.")
+	void EndProcess() {}
+	void EndProcess(TSharedPtr<IPixelCaptureOutputFrame> OutputBuffer);
 
 private:
 	bool bInitialized = false;
@@ -132,14 +151,12 @@ private:
 	int32 ExpectedInputHeight = 0;
 
 	FThreadSafeCounter64 FrameId;
-	uint64 StartTime;
-	uint64 CPUStartTime;
-	uint64 GPUEnqueueTime;
-	uint64 GPUStartTime;
 
 	TUniquePtr<UE::PixelCapture::FOutputFrameBuffer> Buffer;
-	TSharedPtr<IPixelCaptureOutputFrame> CurrentOutputBuffer;
 
-	void InitMetadata(FPixelCaptureFrameMetadata InputFrame);
-	void FinalizeMetadata();
+	mutable FCriticalSection					 OutputBuffersMutex;
+	TArray<TSharedPtr<IPixelCaptureOutputFrame>> OutputBuffers;
+
+	void InitMetadata(FPixelCaptureFrameMetadata InputMetadata, TSharedPtr<IPixelCaptureOutputFrame> OutputBuffer);
+	void FinalizeMetadata(TSharedPtr<IPixelCaptureOutputFrame> OutputBuffer);
 };
diff --git a/Engine/Plugins/Media/PixelCapture/Source/PixelCapture/Public/PixelCaptureCapturerI420.h b/Engine/Plugins/Media/PixelCapture/Source/PixelCapture/Public/PixelCaptureCapturerI420.h
--- a/Engine/Plugins/Media/PixelCapture/Source/PixelCapture/Public/PixelCaptureCapturerI420.h
+++ b/Engine/Plugins/Media/PixelCapture/Source/PixelCapture/Public/PixelCaptureCapturerI420.h
@@ -19,5 +19,5 @@ protected:
 	virtual FString GetCapturerName() const override { return "I420 Copy"; }
 	virtual void Initialize(int32 InputWidth, int32 InputHeight) override;
 	virtual IPixelCaptureOutputFrame* CreateOutputBuffer(int32 InputWidth, int32 InputHeight) override;
-	virtual void BeginProcess(const IPixelCaptureInputFrame& InputFrame, IPixelCaptureOutputFrame* OutputBuffer) override;
+	virtual void BeginProcess(const IPixelCaptureInputFrame& InputFrame, TSharedPtr<IPixelCaptureOutputFrame> OutputBuffer) override;
 };
diff --git a/Engine/Plugins/Media/PixelCapture/Source/PixelCapture/Public/PixelCaptureCapturerI420ToRHI.h b/Engine/Plugins/Media/PixelCapture/Source/PixelCapture/Public/PixelCaptureCapturerI420ToRHI.h
--- a/Engine/Plugins/Media/PixelCapture/Source/PixelCapture/Public/PixelCaptureCapturerI420ToRHI.h
+++ b/Engine/Plugins/Media/PixelCapture/Source/PixelCapture/Public/PixelCaptureCapturerI420ToRHI.h
@@ -19,7 +19,7 @@ protected:
 	virtual FString GetCapturerName() const override { return "I420ToRHI"; }
 	virtual void Initialize(int32 InputWidth, int32 InputHeight) override;
 	virtual IPixelCaptureOutputFrame* CreateOutputBuffer(int32 InputWidth, int32 InputHeight) override;
-	virtual void BeginProcess(const IPixelCaptureInputFrame& InputFrame, IPixelCaptureOutputFrame* OutputBuffer) override;
+	virtual void BeginProcess(const IPixelCaptureInputFrame& InputFrame, TSharedPtr<IPixelCaptureOutputFrame> OutputBuffer) override;
 
 private:
 	FIntPoint Dimensions;
diff --git a/Engine/Plugins/Media/PixelCapture/Source/PixelCapture/Public/PixelCaptureCapturerLayered.h b/Engine/Plugins/Media/PixelCapture/Source/PixelCapture/Public/PixelCaptureCapturerLayered.h
--- a/Engine/Plugins/Media/PixelCapture/Source/PixelCapture/Public/PixelCaptureCapturerLayered.h
+++ b/Engine/Plugins/Media/PixelCapture/Source/PixelCapture/Public/PixelCaptureCapturerLayered.h
@@ -39,7 +39,7 @@ public:
 	 * A callback to broadcast on when the frame has completed the capture process.
 	 * Called once when all layers of a given input frame have completed.
 	 */
-	DECLARE_MULTICAST_DELEGATE(FOnComplete);
+	DECLARE_TS_MULTICAST_DELEGATE(FOnComplete);
 	FOnComplete OnComplete;
 
 protected:
@@ -52,6 +52,5 @@ protected:
 	int32 DestinationFormat;
 	TArray<float> LayerScales;
 	TArray<TSharedPtr<FPixelCaptureCapturer>> LayerCapturers;
-	TAtomic<int> PendingLayers = 0; // atomic because the complete events can come in on multiple threads
 	mutable FCriticalSection LayersGuard;
 };
diff --git a/Engine/Plugins/Media/PixelCapture/Source/PixelCapture/Public/PixelCaptureCapturerNV12ToRHI.h b/Engine/Plugins/Media/PixelCapture/Source/PixelCapture/Public/PixelCaptureCapturerNV12ToRHI.h
--- a/Engine/Plugins/Media/PixelCapture/Source/PixelCapture/Public/PixelCaptureCapturerNV12ToRHI.h
+++ b/Engine/Plugins/Media/PixelCapture/Source/PixelCapture/Public/PixelCaptureCapturerNV12ToRHI.h
@@ -19,7 +19,7 @@ protected:
 	virtual FString GetCapturerName() const override { return "NV12ToRHI"; }
 	virtual void Initialize(int32 InputWidth, int32 InputHeight) override;
 	virtual IPixelCaptureOutputFrame* CreateOutputBuffer(int32 InputWidth, int32 InputHeight) override;
-	virtual void BeginProcess(const IPixelCaptureInputFrame& InputFrame, IPixelCaptureOutputFrame* OutputBuffer) override;
+	virtual void BeginProcess(const IPixelCaptureInputFrame& InputFrame, TSharedPtr<IPixelCaptureOutputFrame> OutputBuffer) override;
 
 private:
 	FIntPoint Dimensions;
diff --git a/Engine/Plugins/Media/PixelCapture/Source/PixelCapture/Public/PixelCaptureCapturerRHI.h b/Engine/Plugins/Media/PixelCapture/Source/PixelCapture/Public/PixelCaptureCapturerRHI.h
--- a/Engine/Plugins/Media/PixelCapture/Source/PixelCapture/Public/PixelCaptureCapturerRHI.h
+++ b/Engine/Plugins/Media/PixelCapture/Source/PixelCapture/Public/PixelCaptureCapturerRHI.h
@@ -21,17 +21,16 @@ public:
 	virtual ~FPixelCaptureCapturerRHI() = default;
 
 protected:
-	virtual FString GetCapturerName() const override { return "RHI Copy"; }
+	virtual FString					  GetCapturerName() const override { return "RHI Copy"; }
 	virtual IPixelCaptureOutputFrame* CreateOutputBuffer(int32 InputWidth, int32 InputHeight) override;
-	virtual void BeginProcess(const IPixelCaptureInputFrame& InputFrame, IPixelCaptureOutputFrame* OutputBuffer) override;
+	virtual void					  BeginProcess(const IPixelCaptureInputFrame& InputFrame, TSharedPtr<IPixelCaptureOutputFrame> OutputBuffer) override;
 
 private:
 	float Scale = 1.0f;
 
-	FGPUFenceRHIRef Fence;
 	ERHIInterfaceType RHIType;
 
 	FPixelCaptureCapturerRHI(float InScale);
-	void CheckComplete();
-	void OnRHIStageComplete();
+	void CheckComplete(TSharedPtr<IPixelCaptureOutputFrame> OutputBuffer, FGPUFenceRHIRef Fence);
+	void OnRHIStageComplete(TSharedPtr<IPixelCaptureOutputFrame> OutputBuffer);
 };
diff --git a/Engine/Plugins/Media/PixelCapture/Source/PixelCapture/Public/PixelCaptureCapturerRHINoCopy.h b/Engine/Plugins/Media/PixelCapture/Source/PixelCapture/Public/PixelCaptureCapturerRHINoCopy.h
--- a/Engine/Plugins/Media/PixelCapture/Source/PixelCapture/Public/PixelCaptureCapturerRHINoCopy.h
+++ b/Engine/Plugins/Media/PixelCapture/Source/PixelCapture/Public/PixelCaptureCapturerRHINoCopy.h
@@ -25,7 +25,7 @@ public:
 protected:
 	virtual FString GetCapturerName() const override { return "RHI No Copy"; }
 	virtual IPixelCaptureOutputFrame* CreateOutputBuffer(int32 InputWidth, int32 InputHeight) override;
-	virtual void BeginProcess(const IPixelCaptureInputFrame& InputFrame, IPixelCaptureOutputFrame* OutputBuffer) override;
+	virtual void BeginProcess(const IPixelCaptureInputFrame& InputFrame, TSharedPtr<IPixelCaptureOutputFrame> OutputBuffer) override;
 
 private:
 	float Scale = 1.0f;
diff --git a/Engine/Plugins/Media/PixelCapture/Source/PixelCapture/Public/PixelCaptureCapturerRHIRDG.h b/Engine/Plugins/Media/PixelCapture/Source/PixelCapture/Public/PixelCaptureCapturerRHIRDG.h
--- a/Engine/Plugins/Media/PixelCapture/Source/PixelCapture/Public/PixelCaptureCapturerRHIRDG.h
+++ b/Engine/Plugins/Media/PixelCapture/Source/PixelCapture/Public/PixelCaptureCapturerRHIRDG.h
@@ -23,12 +23,12 @@ public:
 protected:
 	virtual FString GetCapturerName() const override { return "RHI RDG Copy"; }
 	virtual IPixelCaptureOutputFrame* CreateOutputBuffer(int32 InputWidth, int32 InputHeight) override;
-	virtual void BeginProcess(const IPixelCaptureInputFrame& InputFrame, IPixelCaptureOutputFrame* OutputBuffer) override;
+	virtual void BeginProcess(const IPixelCaptureInputFrame& InputFrame, TSharedPtr<IPixelCaptureOutputFrame> OutputBuffer) override;
 
 private:
 	float Scale = 1.0f;
 
 	FPixelCaptureCapturerRHIRDG(float InScale);
 	void CheckComplete();
-	void OnRHIStageComplete();
+	void OnRHIStageComplete(TSharedPtr<IPixelCaptureOutputFrame> OutputBuffer);
 };
diff --git a/Engine/Plugins/Media/PixelCapture/Source/PixelCapture/Public/PixelCaptureCapturerRHIToI420CPU.h b/Engine/Plugins/Media/PixelCapture/Source/PixelCapture/Public/PixelCaptureCapturerRHIToI420CPU.h
--- a/Engine/Plugins/Media/PixelCapture/Source/PixelCapture/Public/PixelCaptureCapturerRHIToI420CPU.h
+++ b/Engine/Plugins/Media/PixelCapture/Source/PixelCapture/Public/PixelCaptureCapturerRHIToI420CPU.h
@@ -26,7 +26,7 @@ protected:
 	virtual FString					  GetCapturerName() const override { return "RHIToI420CPU"; }
 	virtual void					  Initialize(int32 InputWidth, int32 InputHeight) override;
 	virtual IPixelCaptureOutputFrame* CreateOutputBuffer(int32 InputWidth, int32 InputHeight) override;
-	virtual void					  BeginProcess(const IPixelCaptureInputFrame& InputFrame, IPixelCaptureOutputFrame* OutputBuffer) override;
+	virtual void					  BeginProcess(const IPixelCaptureInputFrame& InputFrame, TSharedPtr<IPixelCaptureOutputFrame> OutputBuffer) override;
 
 private:
 	float Scale = 1.0f;
@@ -37,8 +37,8 @@ private:
 
 #if 1
 	TSharedPtr<FRHIGPUTextureReadback> TextureReader;
-	void OnRHIStageComplete(IPixelCaptureOutputFrame* OutputBuffer);
-	void CheckComplete(IPixelCaptureOutputFrame* OutputBuffer);
+	void OnRHIStageComplete(TSharedPtr<IPixelCaptureOutputFrame> OutputBuffer);
+	void CheckComplete(TSharedPtr<IPixelCaptureOutputFrame> OutputBuffer);
 #else
 	FTextureRHIRef ReadbackTexture;
 #endif
diff --git a/Engine/Plugins/Media/PixelCapture/Source/PixelCapture/Public/PixelCaptureCapturerRHIToI420Compute.h b/Engine/Plugins/Media/PixelCapture/Source/PixelCapture/Public/PixelCaptureCapturerRHIToI420Compute.h
--- a/Engine/Plugins/Media/PixelCapture/Source/PixelCapture/Public/PixelCaptureCapturerRHIToI420Compute.h
+++ b/Engine/Plugins/Media/PixelCapture/Source/PixelCapture/Public/PixelCaptureCapturerRHIToI420Compute.h
@@ -25,7 +25,7 @@ protected:
 	virtual FString GetCapturerName() const override { return "RHIToI420Compute"; }
 	virtual void Initialize(int32 InputWidth, int32 InputHeight) override;
 	virtual IPixelCaptureOutputFrame* CreateOutputBuffer(int32 InputWidth, int32 InputHeight) override;
-	virtual void BeginProcess(const IPixelCaptureInputFrame& InputFrame, IPixelCaptureOutputFrame* OutputBuffer) override;
+	virtual void BeginProcess(const IPixelCaptureInputFrame& InputFrame, TSharedPtr<IPixelCaptureOutputFrame> OutputBuffer) override;
 
 private:
 	float Scale = 1.0f;
@@ -59,6 +59,6 @@ private:
 	int32 VStride = 0;
 
 	FPixelCaptureCapturerRHIToI420Compute(float InScale);
-	void OnRHIStageComplete(IPixelCaptureOutputFrame* OutputBuffer);
+	void OnRHIStageComplete(TSharedPtr<IPixelCaptureOutputFrame> OutputBuffer);
 	void CleanUp();
 };
diff --git a/Engine/Plugins/Media/PixelCapture/Source/PixelCapture/Public/PixelCaptureFrameMetadata.h b/Engine/Plugins/Media/PixelCapture/Source/PixelCapture/Public/PixelCaptureFrameMetadata.h
--- a/Engine/Plugins/Media/PixelCapture/Source/PixelCapture/Public/PixelCaptureFrameMetadata.h
+++ b/Engine/Plugins/Media/PixelCapture/Source/PixelCapture/Public/PixelCaptureFrameMetadata.h
@@ -4,11 +4,9 @@
 
 #include "CoreMinimal.h"
 
-class FPixelCaptureUserData
-{
-};
+class PIXELCAPTURE_API FPixelCaptureUserData{};
 
-struct FPixelCaptureFrameMetadata
+struct PIXELCAPTURE_API FPixelCaptureFrameMetadata
 {
 	// Identifier for the capture pipeline/process this frame took
 	FString ProcessName = "Unknown";
@@ -22,10 +20,34 @@ struct FPixelCaptureFrameMetadata
 	// The time this frame was sourced/created
 	uint64 SourceTime = 0;
 
+	// Timestamps to track the entire length of the capture process
+	uint64 CaptureStartCyles = 0;
+	uint64 CaptureEndCyles = 0;
+
+	// Timestamps for when CPU work has started / ended
+	uint64 CaptureProcessCPUStartCycles = 0;
+	uint64 CaptureProcessCPUEndCycles = 0;
+
+	// Timestamps for when the capture process was enqueued to the GPU and when the GPU actually started work
+	uint64 CaptureProcessGPUEnqueueStartCycles = 0;
+	uint64 CaptureProcessGPUEnqueueEndCycles = 0;
+
+	// Timestamps for when the GPU started work and when the GPU finished the capture process
+	uint64 CaptureProcessGPUStartCycles = 0;
+	uint64 CaptureProcessGPUEndCycles = 0;
+
+	// Timestamps for tracking time after the GPU completed it's work until the capture process completed.
+	uint64 CaptureProcessPostGPUStartCycles = 0;
+	uint64 CaptureProcessPostGPUEndCycles = 0;
+
 	// Capture process timings. Duration not timestamp.
+	UE_DEPRECATED(5.6, "CaptureTime has been deprecated. You can calculate the time taken with FPlatformTime::ToMilliseconds64(CaptureEndCyles - CaptureStartCyles).")
 	uint64 CaptureTime = 0;
+	UE_DEPRECATED(5.6, "CaptureProcessCPUTime has been deprecated. You can calculate the time taken with FPlatformTime::ToMilliseconds64(CaptureProcessCPUEndCycles - CaptureProcessCPUStartCycles).")
 	uint64 CaptureProcessCPUTime = 0;
+	UE_DEPRECATED(5.6, "CaptureProcessGPUDelay has been deprecated. You can calculate the time taken with FPlatformTime::ToMilliseconds64(CaptureProcessGPUEnqueueEndCycles - CaptureProcessGPUEnqueueStartCycles).")
 	uint64 CaptureProcessGPUDelay = 0;
+	UE_DEPRECATED(5.6, "CaptureProcessGPUTime has been deprecated. You can calculate the time taken with FPlatformTime::ToMilliseconds64(CaptureProcessGPUEndCycles - CaptureProcessGPUStartCycles).")
 	uint64 CaptureProcessGPUTime = 0;
 
 	// Display process timings. Duration not timestamp.
@@ -52,6 +74,18 @@ struct FPixelCaptureFrameMetadata
 	// Pointer to a user class that stores additional metadata that they want to keep throughout the process
 	TSharedPtr<FPixelCaptureUserData> UserData;
 
+	// Disable deprecations about accessing the deprecated members during copy
+	PRAGMA_DISABLE_DEPRECATION_WARNINGS
 	// wanted this to be explicit with a name
-	FPixelCaptureFrameMetadata Copy() const { return *this; }
+	FPixelCaptureFrameMetadata Copy() const
+	{
+		return *this;
+	}
+
+	FPixelCaptureFrameMetadata() = default;
+	FPixelCaptureFrameMetadata(FPixelCaptureFrameMetadata&&) = default;
+	FPixelCaptureFrameMetadata(const FPixelCaptureFrameMetadata&) = default;
+	FPixelCaptureFrameMetadata& operator=(FPixelCaptureFrameMetadata&&) = default;
+	FPixelCaptureFrameMetadata& operator=(const FPixelCaptureFrameMetadata&) = default;
+	PRAGMA_ENABLE_DEPRECATION_WARNINGS
 };
diff --git a/Engine/Plugins/Media/PixelStreaming/Source/PixelStreaming/Private/Stats.cpp b/Engine/Plugins/Media/PixelStreaming/Source/PixelStreaming/Private/Stats.cpp
--- a/Engine/Plugins/Media/PixelStreaming/Source/PixelStreaming/Private/Stats.cpp
+++ b/Engine/Plugins/Media/PixelStreaming/Source/PixelStreaming/Private/Stats.cpp
@@ -631,10 +631,10 @@ namespace UE::PixelStreaming
 
 	void FStats::AddFrameTimingStats(const FPixelCaptureFrameMetadata& FrameMetadata)
 	{
-		const double TimeCapture = AddTimeStat(FrameMetadata.CaptureTime, FString::Printf(TEXT("%s Layer %d Frame Capture Time"), *FrameMetadata.ProcessName, FrameMetadata.Layer));
-		const double TimeCPU = AddTimeStat(FrameMetadata.CaptureProcessCPUTime, FString::Printf(TEXT("%s Layer %d Frame Capture CPU Time"), *FrameMetadata.ProcessName, FrameMetadata.Layer));
-		const double TimeGPUDelay = AddTimeStat(FrameMetadata.CaptureProcessGPUDelay, FString::Printf(TEXT("%s Layer %d Frame Capture GPU Delay Time"), *FrameMetadata.ProcessName, FrameMetadata.Layer));
-		const double TimeGPU = AddTimeStat(FrameMetadata.CaptureProcessGPUTime, FString::Printf(TEXT("%s Layer %d Frame Capture GPU Time"), *FrameMetadata.ProcessName, FrameMetadata.Layer));
+		const double TimeCapture = AddTimeStat(FPlatformTime::ToMilliseconds64(FrameMetadata.CaptureEndCyles - FrameMetadata.CaptureStartCyles), FString::Printf(TEXT("%s Layer %d Frame Capture Time"), *FrameMetadata.ProcessName, FrameMetadata.Layer));
+		const double TimeCPU = AddTimeStat(FPlatformTime::ToMilliseconds64(FrameMetadata.CaptureProcessCPUEndCycles - FrameMetadata.CaptureProcessCPUStartCycles), FString::Printf(TEXT("%s Layer %d Frame Capture CPU Time"), *FrameMetadata.ProcessName, FrameMetadata.Layer));
+		const double TimeGPUDelay = AddTimeStat(FPlatformTime::ToMilliseconds64(FrameMetadata.CaptureProcessGPUEnqueueEndCycles - FrameMetadata.CaptureProcessGPUEnqueueStartCycles), FString::Printf(TEXT("%s Layer %d Frame Capture GPU Delay Time"), *FrameMetadata.ProcessName, FrameMetadata.Layer));
+		const double TimeGPU = AddTimeStat(FPlatformTime::ToMilliseconds64(FrameMetadata.CaptureProcessGPUEndCycles - FrameMetadata.CaptureProcessGPUStartCycles), FString::Printf(TEXT("%s Layer %d Frame Capture GPU Time"), *FrameMetadata.ProcessName, FrameMetadata.Layer));
 		const double TimeEncode = AddTimeDeltaStat(FrameMetadata.LastEncodeEndTime, FrameMetadata.LastEncodeStartTime, FString::Printf(TEXT("%s Layer %d Frame Encode Time"), *FrameMetadata.ProcessName, FrameMetadata.Layer));
 		const double TimePacketize = AddTimeDeltaStat(FrameMetadata.LastPacketizationEndTime, FrameMetadata.LastPacketizationStartTime, FString::Printf(TEXT("%s Layer %d Frame Packetization Time"), *FrameMetadata.ProcessName, FrameMetadata.Layer));
 
diff --git a/Engine/Plugins/Media/PixelStreaming2/Source/PixelStreaming2RTC/Private/Stats.cpp b/Engine/Plugins/Media/PixelStreaming2/Source/PixelStreaming2RTC/Private/Stats.cpp
--- a/Engine/Plugins/Media/PixelStreaming2/Source/PixelStreaming2RTC/Private/Stats.cpp
+++ b/Engine/Plugins/Media/PixelStreaming2/Source/PixelStreaming2RTC/Private/Stats.cpp
@@ -598,19 +598,21 @@ namespace UE::PixelStreaming2
 			GraphValue(*FString::Printf(TEXT("%d Produce Time"), FrameMetadata.Layer), StaticCast<float>(TimeProduce), Samples, 0.0f, 30.0f);
 		}
 
-		const double TimeCapture = AddTimeStat(FrameMetadata.CaptureTime, FString::Printf(TEXT("%s Layer %d Frame Capture Time"), *FrameMetadata.ProcessName, FrameMetadata.Layer));
-		const double TimeCPU = AddTimeStat(FrameMetadata.CaptureProcessCPUTime, FString::Printf(TEXT("%s Layer %d Frame Capture CPU Time"), *FrameMetadata.ProcessName, FrameMetadata.Layer));
-		const double TimeGPUDelay = AddTimeStat(FrameMetadata.CaptureProcessGPUDelay, FString::Printf(TEXT("%s Layer %d Frame Capture GPU Delay Time"), *FrameMetadata.ProcessName, FrameMetadata.Layer));
-		const double TimeGPU = AddTimeStat(FrameMetadata.CaptureProcessGPUTime, FString::Printf(TEXT("%s Layer %d Frame Capture GPU Time"), *FrameMetadata.ProcessName, FrameMetadata.Layer));
+		const double TimeCapture = AddTimeStat(FPlatformTime::ToMilliseconds64(FrameMetadata.CaptureEndCyles - FrameMetadata.CaptureStartCyles), FString::Printf(TEXT("%s Layer %d Frame Capture Time"), *FrameMetadata.ProcessName, FrameMetadata.Layer));
+		const double TimeCPU = AddTimeStat(FPlatformTime::ToMilliseconds64(FrameMetadata.CaptureProcessCPUEndCycles - FrameMetadata.CaptureProcessCPUStartCycles), FString::Printf(TEXT("%s Layer %d Frame Capture CPU Time"), *FrameMetadata.ProcessName, FrameMetadata.Layer));
+		const double TimeGPUDelay = AddTimeStat(FPlatformTime::ToMilliseconds64(FrameMetadata.CaptureProcessGPUEnqueueEndCycles - FrameMetadata.CaptureProcessGPUEnqueueStartCycles), FString::Printf(TEXT("%s Layer %d Frame Capture GPU Delay Time"), *FrameMetadata.ProcessName, FrameMetadata.Layer));
+		const double TimeGPU = AddTimeStat(FPlatformTime::ToMilliseconds64(FrameMetadata.CaptureProcessGPUEndCycles - FrameMetadata.CaptureProcessGPUStartCycles), FString::Printf(TEXT("%s Layer %d Frame Capture GPU Time"), *FrameMetadata.ProcessName, FrameMetadata.Layer));
+		const double TimePostGPU = AddTimeStat(FPlatformTime::ToMilliseconds64(FrameMetadata.CaptureProcessPostGPUEndCycles - FrameMetadata.CaptureProcessPostGPUStartCycles), FString::Printf(TEXT("%s Layer %d Frame Capture Post GPU Time"), *FrameMetadata.ProcessName, FrameMetadata.Layer));
 
 		const FStat UseData({ .Name = FName(*FString::Printf(TEXT("%s Layer %d Frame Uses"), *FrameMetadata.ProcessName, FrameMetadata.Layer)) }, static_cast<double>(FrameMetadata.UseCount));
 		StoreApplicationStat(UseData);
 
-		GraphValue(*FString::Printf(TEXT("%d Capture Time"), FrameMetadata.Layer), StaticCast<float>(TimeCapture), Samples, 0.0f, 30.0f);
-		GraphValue(*FString::Printf(TEXT("%d CPU Time"), FrameMetadata.Layer), StaticCast<float>(TimeCPU), Samples, 0.0f, 30.0f);
-		GraphValue(*FString::Printf(TEXT("%d GPU Delay Time"), FrameMetadata.Layer), StaticCast<float>(TimeGPUDelay), Samples, 0.0f, 30.0f);
-		GraphValue(*FString::Printf(TEXT("%d GPU Time"), FrameMetadata.Layer), StaticCast<float>(TimeGPU), Samples, 0.0f, 30.0f);
-		GraphValue(*FString::Printf(TEXT("%d Frame Uses"), FrameMetadata.Layer), StaticCast<float>(FrameMetadata.UseCount), Samples, 0.0f, 10.0f);
+		GraphValue(*FString::Printf(TEXT("Layer %d Capture Time"), FrameMetadata.Layer), StaticCast<float>(TimeCapture), Samples, 0.0f, 30.0f);
+		GraphValue(*FString::Printf(TEXT("Layer %d CPU Time"), FrameMetadata.Layer), StaticCast<float>(TimeCPU), Samples, 0.0f, 30.0f);
+		GraphValue(*FString::Printf(TEXT("Layer %d GPU Delay Time"), FrameMetadata.Layer), StaticCast<float>(TimeGPUDelay), Samples, 0.0f, 30.0f);
+		GraphValue(*FString::Printf(TEXT("Layer %d GPU Time"), FrameMetadata.Layer), StaticCast<float>(TimeGPU), Samples, 0.0f, 30.0f);
+		GraphValue(*FString::Printf(TEXT("Layer %d Post GPU Time"), FrameMetadata.Layer), StaticCast<float>(TimePostGPU), Samples, 0.0f, 30.0f);
+		GraphValue(*FString::Printf(TEXT("Layer %d Frame Uses"), FrameMetadata.Layer), StaticCast<float>(FrameMetadata.UseCount), Samples, 0.0f, 10.0f);
 	}
 
 	void FStats::AddCanvasTile(FName Name, const FCanvasTileItem& Tile)
