[PixelCapture] Fix Capturing VPX Capture on Mac with MediaCapture.

The problem is MediaCapture needs to do the capture on the rendering thread. This change forces Mac to use the rendering thread with cpu frames.

diff --git a/Engine/Plugins/Media/PixelCapture/Source/PixelCapture/Internal/PixelCaptureCapturerMediaCapture.h b/Engine/Plugins/Media/PixelCapture/Source/PixelCapture/Internal/PixelCaptureCapturerMediaCapture.h
--- a/Engine/Plugins/Media/PixelCapture/Source/PixelCapture/Internal/PixelCaptureCapturerMediaCapture.h
+++ b/Engine/Plugins/Media/PixelCapture/Source/PixelCapture/Internal/PixelCaptureCapturerMediaCapture.h
@@ -76,6 +76,14 @@ protected:
 		FVector2D											   CropU,
 		FVector2D											   CropV) override;
 
+	virtual void OnFrameCaptured_RenderingThread(
+		const FCaptureBaseData& InBaseData,
+		TSharedPtr<FMediaCaptureUserData,ESPMode::ThreadSafe> InUserData,
+		void* InBuffer,
+		int32 Width,
+		int32 Height,
+		int32 BytesPerRow) override;
+
 	virtual bool ShouldCaptureRHIResource() const override
 	{
 		return Format == PixelCaptureBufferFormat::FORMAT_RHI;
@@ -83,7 +91,14 @@ protected:
 
 	virtual bool SupportsAnyThreadCapture() const
 	{
+#if PLATFORM_MAC
+		// On Mac, Media Capture must capture cpu frames on the render thread. 
+		// This is because MediaCapture Lock_Unsafe uses RHIMapStagingSurface 
+		// which needs to run on the rendering thread with Metal.
+		return Format == PixelCaptureBufferFormat::FORMAT_RHI;
+#else // PLATFORM_MAC
 		return true;
+#endif // PLATFORM_MAC
 	}
 
 	virtual ETextureCreateFlags GetOutputTextureFlags() const override;
diff --git a/Engine/Plugins/Media/PixelCapture/Source/PixelCapture/Private/PixelCaptureCapturerMediaCapture.cpp b/Engine/Plugins/Media/PixelCapture/Source/PixelCapture/Private/PixelCaptureCapturerMediaCapture.cpp
--- a/Engine/Plugins/Media/PixelCapture/Source/PixelCapture/Private/PixelCaptureCapturerMediaCapture.cpp
+++ b/Engine/Plugins/Media/PixelCapture/Source/PixelCapture/Private/PixelCaptureCapturerMediaCapture.cpp
@@ -63,6 +63,18 @@ void UPixelCaptureMediaCapture::OnFrameCaptured_AnyThread(
 	OnCaptureComplete.Broadcast(OutputFrame);
 }
 
+void UPixelCaptureMediaCapture::OnFrameCaptured_RenderingThread(
+	const FCaptureBaseData&								   InBaseData,
+	TSharedPtr<FMediaCaptureUserData, ESPMode::ThreadSafe> InUserData,
+	void*												   InBuffer,
+	int32												   Width,
+	int32												   Height,
+	int32												   BytesPerRow)
+{
+	FMediaCaptureResourceData ResourceData{ .Buffer = InBuffer, .Width = Width, .Height = Height, .BytesPerRow = BytesPerRow };
+	OnFrameCaptured_AnyThread(InBaseData, InUserData, ResourceData);
+}
+
 void UPixelCaptureMediaCapture::OnCustomCapture_RenderingThread(
 	FRDGBuilder&										   GraphBuilder,
 	const FCaptureBaseData&								   InBaseData,
