[PixelCapture] Fix Pixel Capture Tests

PixelCapture tests failing after recent PixelCapture Changes as the expected behavour changed.

diff --git a/Engine/Plugins/Media/PixelCapture/Source/PixelCapture/Private/Tests/PixelCaptureTests.cpp b/Engine/Plugins/Media/PixelCapture/Source/PixelCapture/Private/Tests/PixelCaptureTests.cpp
--- a/Engine/Plugins/Media/PixelCapture/Source/PixelCapture/Private/Tests/PixelCaptureTests.cpp
+++ b/Engine/Plugins/Media/PixelCapture/Source/PixelCapture/Private/Tests/PixelCaptureTests.cpp
@@ -83,6 +83,10 @@ namespace
 				EndProcess(OutputBuffer);
 				SetIsBusy(false);
 			}
+			else
+			{
+				StoredOutputBuffer = OutputBuffer;
+			}
 		}
 
 		void MockProcessComplete(TSharedPtr<IPixelCaptureOutputFrame> OutputBuffer)
@@ -90,6 +94,8 @@ namespace
 			EndProcess(OutputBuffer);
 			SetIsBusy(false);
 		}
+
+		TSharedPtr<IPixelCaptureOutputFrame> StoredOutputBuffer;
 	};
 
 	struct FMockCaptureSource : public IPixelCaptureCapturerSource
@@ -148,8 +154,7 @@ namespace UE::PixelCapture
 			TestTrue("Capturer busy after frame input", Capturer.IsBusy() == true);
 			TestTrue("Capturer has no output before completing", Capturer.HasOutput() == false);
 			TestTrue("Capturer still returns null after input but not complete", Capturer.ReadOutput() == nullptr);
-			TSharedPtr<IPixelCaptureOutputFrame> OutbutBuffer = MakeShareable<IPixelCaptureOutputFrame>(Capturer.CreateOutputBuffer(1920, 1080));
-			Capturer.MockProcessComplete(OutbutBuffer);
+			Capturer.MockProcessComplete(Capturer.StoredOutputBuffer);
 			TestTrue("Capturer still initialized after process complete", Capturer.IsInitialized() == true);
 			TestTrue("Capturer no longer busy after process complete", Capturer.IsBusy() == false);
 			TestTrue("Capturer has output after completing", Capturer.HasOutput() == true);
@@ -163,7 +168,7 @@ namespace UE::PixelCapture
 			TSharedPtr<FPixelCaptureCapturerLayered> LayeredCapturer = FPixelCaptureCapturerLayered::Create(&CaptureSource, StaticCast<int32>(EMockFormats::Output1), Scales);
 			LayeredCapturer->OnComplete.AddLambda([&]() { CaptureSource.FrameCompleteCount++; });
 			LayeredCapturer->Capture(FMockInput(InputWidth, InputHeight, InputMagicData));
-			TestTrue("Frame complete callback called once.", CaptureSource.FrameCompleteCount == 1);
+			TestTrue("Frame complete callback called once.", CaptureSource.FrameCompleteCount == 3);
 			for (TMap<FMockCapturer*, int>::ElementType& CaptureComplete : CaptureSource.LayerCompleteCounts)
 			{
 				TestTrue(FString::Printf(TEXT("Layer %p called complete callback once."), CaptureComplete.Key), CaptureComplete.Value == 1);
@@ -184,6 +189,7 @@ namespace UE::PixelCapture
 
 			// test that WaitForFormat will time out
 			{
+				AddExpectedMessage(TEXT("No capturer configured for LayerIndex (0)!"), EAutomationExpectedMessageFlags::MatchType::Exact, 1, false);
 				const double PreWaitSeconds = FPlatformTime::Seconds();
 				FormatCapturer->WaitForFormat(TestFormat, 0, 500);
 				const double DeltaSeconds = FPlatformTime::Seconds() - PreWaitSeconds;
@@ -250,7 +256,7 @@ namespace UE::PixelCapture
 			FormatCapturer->AddOutputFormat(StaticCast<int>(EMockFormats::Output2));
 			FormatCapturer->Capture(FMockInput(InputWidth, InputHeight, InputMagicData));
 
-			TestTrue("MultiFormat Frame complete callback called once.", CaptureSource.FrameCompleteCount == 1);
+			TestTrue("MultiFormat Frame complete callback called once.", CaptureSource.FrameCompleteCount == 6);
 			for (TMap<FMockCapturer*, int>::ElementType& CaptureComplete : CaptureSource.LayerCompleteCounts)
 			{
 				TestTrue(FString::Printf(TEXT("Format %d Layer %p called complete callback once."), CaptureComplete.Key->Format, CaptureComplete.Key), CaptureComplete.Value == 1);
