[PixelCapture] Fix Pixel Streaming log spam when using MediaCapture.

MediaCapture can only capture the current frame once in the render thread. If it attempts to capture a frame a second time, it will log a warning. The reason why it may try to capture twice is if the game thread captures frames faster than the render thread.

This fix ensures that PixelCapture does not call MediaCapture more than once per frame on the renderthread.

diff --git a/Engine/Plugins/Media/PixelCapture/Source/PixelCapture/Internal/PixelCaptureCapturerMediaCapture.h b/Engine/Plugins/Media/PixelCapture/Source/PixelCapture/Internal/PixelCaptureCapturerMediaCapture.h
--- a/Engine/Plugins/Media/PixelCapture/Source/PixelCapture/Internal/PixelCaptureCapturerMediaCapture.h
+++ b/Engine/Plugins/Media/PixelCapture/Source/PixelCapture/Internal/PixelCaptureCapturerMediaCapture.h
@@ -133,6 +133,7 @@ protected:
 	virtual void					  BeginProcess(const IPixelCaptureInputFrame& InputFrame, TSharedPtr<IPixelCaptureOutputFrame> OutputBuffer) override;
 
 private:
+	uint64 LastFrameCounterRenderThread = 0;
 	float Scale = 1.0f;
 	int32 Format = PixelCaptureBufferFormat::FORMAT_UNKNOWN;
 	bool  bMediaCaptureInitialized = false;
diff --git a/Engine/Plugins/Media/PixelCapture/Source/PixelCapture/Private/PixelCaptureCapturerMediaCapture.cpp b/Engine/Plugins/Media/PixelCapture/Source/PixelCapture/Private/PixelCaptureCapturerMediaCapture.cpp
--- a/Engine/Plugins/Media/PixelCapture/Source/PixelCapture/Private/PixelCaptureCapturerMediaCapture.cpp
+++ b/Engine/Plugins/Media/PixelCapture/Source/PixelCapture/Private/PixelCaptureCapturerMediaCapture.cpp
@@ -309,6 +309,15 @@ void FPixelCaptureCapturerMediaCapture::BeginProcess(const IPixelCaptureInputFra
 		return;
 	}
 
+	check(IsInRenderingThread());
+	if (LastFrameCounterRenderThread == GFrameCounterRenderThread)
+	{
+		// Early out as we've already called TryCaptureImmediate_RenderThread for this frame
+		EndProcess(OutputBuffer);
+		return;
+	}
+	LastFrameCounterRenderThread = GFrameCounterRenderThread;
+
 	checkf(InputFrame.GetType() == StaticCast<int32>(PixelCaptureBufferFormat::FORMAT_RHI), TEXT("Incorrect source frame coming into frame capture process."));
 	const FPixelCaptureInputFrameRHI& SourceFrame = StaticCast<const FPixelCaptureInputFrameRHI&>(InputFrame);
 
