[PixelCapture, PixelStreaming, PixelStreaming2] Fix: Frame stutters and memory leaks when using VPX

diff --git a/Engine/Plugins/Media/PixelCapture/Source/PixelCapture/Internal/PixelCaptureCapturerMediaCapture.h b/Engine/Plugins/Media/PixelCapture/Source/PixelCapture/Internal/PixelCaptureCapturerMediaCapture.h
--- a/Engine/Plugins/Media/PixelCapture/Source/PixelCapture/Internal/PixelCaptureCapturerMediaCapture.h
+++ b/Engine/Plugins/Media/PixelCapture/Source/PixelCapture/Internal/PixelCaptureCapturerMediaCapture.h
@@ -147,6 +147,8 @@ private:
 	FPixelCaptureCapturerMediaCapture(FPixelCaptureCapturerConfig& Config, int32 InFormat);
 	void InitializeMediaCapture();
 
+	void InvalidateOutputBuffer(TSharedPtr<IPixelCaptureOutputFrame> OutputBuffer);
+
 	UPROPERTY(Transient)
 	TObjectPtr<UPixelCaptureMediaCapture> MediaCapture = nullptr;
 
diff --git a/Engine/Plugins/Media/PixelCapture/Source/PixelCapture/Private/PixelCaptureCapturerMediaCapture.cpp b/Engine/Plugins/Media/PixelCapture/Source/PixelCapture/Private/PixelCaptureCapturerMediaCapture.cpp
--- a/Engine/Plugins/Media/PixelCapture/Source/PixelCapture/Private/PixelCaptureCapturerMediaCapture.cpp
+++ b/Engine/Plugins/Media/PixelCapture/Source/PixelCapture/Private/PixelCaptureCapturerMediaCapture.cpp
@@ -86,7 +86,7 @@ void UPixelCaptureMediaCapture::OnCustomCapture_RenderingThread(
 	FVector2D											   CropV)
 {
 	FGPUFenceRHIRef Fence = GDynamicRHI->RHICreateGPUFence(TEXT("UPixelCaptureMediaCapture Fence"));
-	if (UseExperimentalScheduling())
+	if (UseExperimentalScheduling() && ShouldCaptureRHIResource())
 	{
 		Fences.PushLast(Fence);
 	}
@@ -308,6 +308,12 @@ void FPixelCaptureCapturerMediaCapture::BeginProcess(const IPixelCaptureInputFra
 {
 	if (!bMediaCaptureInitialized)
 	{
+		UE::PixelCapture::MarkCPUWorkStart(OutputBuffer);
+		UE::PixelCapture::MarkCPUWorkEnd(OutputBuffer);
+		UE::PixelCapture::MarkGPUWorkStart(OutputBuffer);
+		UE::PixelCapture::MarkGPUWorkEnd(OutputBuffer);
+		// No passes added. Invalidate the output buffer so that the encoder won't pull old data from the ring buffer
+		InvalidateOutputBuffer(OutputBuffer);
 		// Early out as media capture is still initializing itself. We'll capture a later frame
 		EndProcess(OutputBuffer);
 		return;
@@ -316,6 +322,12 @@ void FPixelCaptureCapturerMediaCapture::BeginProcess(const IPixelCaptureInputFra
 	check(IsInRenderingThread());
 	if (LastFrameCounterRenderThread == GFrameCounterRenderThread)
 	{
+		UE::PixelCapture::MarkCPUWorkStart(OutputBuffer);
+		UE::PixelCapture::MarkCPUWorkEnd(OutputBuffer);
+		UE::PixelCapture::MarkGPUWorkStart(OutputBuffer);
+		UE::PixelCapture::MarkGPUWorkEnd(OutputBuffer);
+		// No passes added. Invalidate the output buffer so that the encoder won't pull old data from the ring buffer
+		InvalidateOutputBuffer(OutputBuffer);
 		// Early out as we've already called TryCaptureImmediate_RenderThread for this frame
 		EndProcess(OutputBuffer);
 		return;
@@ -361,7 +373,25 @@ void FPixelCaptureCapturerMediaCapture::BeginProcess(const IPixelCaptureInputFra
 	{
 		// TryCaptureImmediate_RenderThread didn't add any passes so remove the latest enqueued output frame from the MediaCapture
 		MediaCapture->RemoveLastOutputFrame();
+		// No passes added. Invalidate the output buffer so that the encoder won't pull old data from the ring buffer
+		InvalidateOutputBuffer(OutputBuffer);
 		// RDG graph had no passes so we can manually call EndProcess()
 		EndProcess(OutputBuffer);
 	}
+}
+
+void FPixelCaptureCapturerMediaCapture::InvalidateOutputBuffer(TSharedPtr<IPixelCaptureOutputFrame> OutputBuffer)
+{
+	if (Format == PixelCaptureBufferFormat::FORMAT_RHI)
+	{
+		StaticCastSharedPtr<FPixelCaptureOutputFrameRHI>(OutputBuffer)->SetFrameTexture(nullptr);
+	}
+	else if (Format == PixelCaptureBufferFormat::FORMAT_I420)
+	{
+		StaticCastSharedPtr<FPixelCaptureOutputFrameI420>(OutputBuffer)->SetI420Buffer(nullptr);
+	}
+	else
+	{
+		UE_LOGFMT(LogPixelCapture, Error, "FPixelCaptureCapturerMediaCapture: Invalid pixel format. Expected either FORMAT_RHI or FORMAT_I420");
+	}
 }
\ No newline at end of file
diff --git a/Engine/Plugins/Media/PixelStreaming/Source/PixelStreaming/Private/PixelStreamingMediaIOCapture.cpp b/Engine/Plugins/Media/PixelStreaming/Source/PixelStreaming/Private/PixelStreamingMediaIOCapture.cpp
--- a/Engine/Plugins/Media/PixelStreaming/Source/PixelStreaming/Private/PixelStreamingMediaIOCapture.cpp
+++ b/Engine/Plugins/Media/PixelStreaming/Source/PixelStreaming/Private/PixelStreamingMediaIOCapture.cpp
@@ -54,7 +54,7 @@ void UPixelStreamingMediaIOCapture::OnCustomCapture_RenderingThread(
 	FVector2D											   CropV)
 {
 	FGPUFenceRHIRef Fence = GDynamicRHI->RHICreateGPUFence(TEXT("UPixelStreamingMediaIOCapture Fence"));
-	if (UseExperimentalScheduling())
+	if (UseExperimentalScheduling() && ShouldCaptureRHIResource())
 	{
 		Fences.PushLast(Fence);
 	}
diff --git a/Engine/Plugins/Media/PixelStreaming2/Source/PixelStreaming2/Private/PixelStreaming2MediaIOCapture.cpp b/Engine/Plugins/Media/PixelStreaming2/Source/PixelStreaming2/Private/PixelStreaming2MediaIOCapture.cpp
--- a/Engine/Plugins/Media/PixelStreaming2/Source/PixelStreaming2/Private/PixelStreaming2MediaIOCapture.cpp
+++ b/Engine/Plugins/Media/PixelStreaming2/Source/PixelStreaming2/Private/PixelStreaming2MediaIOCapture.cpp
@@ -31,7 +31,7 @@ void UPixelStreaming2MediaIOCapture::OnCustomCapture_RenderingThread(
 	FVector2D											   CropV)
 {
 	FGPUFenceRHIRef Fence = GDynamicRHI->RHICreateGPUFence(TEXT("UPixelStreaming2MediaIOCapture Fence"));
-	if (UseExperimentalScheduling())
+	if (UseExperimentalScheduling() && ShouldCaptureRHIResource())
 	{
 		Fences.PushLast(Fence);
 	}
