[PixelCapture] Fix Pixel Streaming using media capturer with DirectX 11.

DirectX 11 does not permit using the experimental scheduler with media capturer, so results in OnRHIResourceCaptured_RenderingThread being called rather than OnRHIResourceCaptured_AnyThread.

This fix adds OnRHIResourceCaptured_RenderingThread for PixelCapturer.

diff --git a/Engine/Plugins/Media/PixelCapture/Source/PixelCapture/Internal/PixelCaptureCapturerMediaCapture.h b/Engine/Plugins/Media/PixelCapture/Source/PixelCapture/Internal/PixelCaptureCapturerMediaCapture.h
--- a/Engine/Plugins/Media/PixelCapture/Source/PixelCapture/Internal/PixelCaptureCapturerMediaCapture.h
+++ b/Engine/Plugins/Media/PixelCapture/Source/PixelCapture/Internal/PixelCaptureCapturerMediaCapture.h
@@ -66,6 +66,12 @@ protected:
 		TSharedPtr<FMediaCaptureUserData, ESPMode::ThreadSafe> UserData,
 		FTextureRHIRef										   Texture) override;
 
+	virtual void OnRHIResourceCaptured_RenderingThread(
+		FRHICommandListImmediate& RHICmdList,
+		const FCaptureBaseData& BaseData,
+		TSharedPtr<FMediaCaptureUserData, ESPMode::ThreadSafe> UserData,
+		FTextureRHIRef Texture) override;
+
 	virtual void OnFrameCaptured_AnyThread(
 		const FCaptureBaseData&								   BaseData,
 		TSharedPtr<FMediaCaptureUserData, ESPMode::ThreadSafe> UserData,
diff --git a/Engine/Plugins/Media/PixelCapture/Source/PixelCapture/Private/PixelCaptureCapturerMediaCapture.cpp b/Engine/Plugins/Media/PixelCapture/Source/PixelCapture/Private/PixelCaptureCapturerMediaCapture.cpp
--- a/Engine/Plugins/Media/PixelCapture/Source/PixelCapture/Private/PixelCaptureCapturerMediaCapture.cpp
+++ b/Engine/Plugins/Media/PixelCapture/Source/PixelCapture/Private/PixelCaptureCapturerMediaCapture.cpp
@@ -31,6 +31,15 @@ void UPixelCaptureMediaCapture::OnRHIResourceCaptured_AnyThread(
 	OnCaptureComplete.Broadcast(OutputFrame);
 }
 
+void UPixelCaptureMediaCapture::OnRHIResourceCaptured_RenderingThread(
+	FRHICommandListImmediate& RHICmdList,
+	const FCaptureBaseData& BaseData,
+	TSharedPtr<FMediaCaptureUserData, ESPMode::ThreadSafe> UserData,
+	FTextureRHIRef Texture)
+{
+	OnRHIResourceCaptured_AnyThread(BaseData, UserData, Texture);
+}
+
 void UPixelCaptureMediaCapture::OnFrameCaptured_AnyThread(
 	const FCaptureBaseData&								   BaseData,
 	TSharedPtr<FMediaCaptureUserData, ESPMode::ThreadSafe> UserData,
