[PixelCapture] Fix Pixel Capture incrollectly incrementing loop in FOutputFrameBuffer::LockProduceBuffer

Changed it to increment by 1, rather than adding "i"

diff --git a/Engine/Plugins/Media/PixelCapture/Source/PixelCapture/Private/OuputFrameBuffer.cpp b/Engine/Plugins/Media/PixelCapture/Source/PixelCapture/Private/OuputFrameBuffer.cpp
--- a/Engine/Plugins/Media/PixelCapture/Source/PixelCapture/Private/OuputFrameBuffer.cpp
+++ b/Engine/Plugins/Media/PixelCapture/Source/PixelCapture/Private/OuputFrameBuffer.cpp
@@ -28,22 +28,24 @@ namespace UE::PixelCapture
 	{
 		checkf(BufferRing.Num() > 0, TEXT("Ring is empty."));
 
-		bool Full = true;
+		bool bIsFull = true;
 
 		// find an unreferenced buffer. by adding one to produce index repeatedly
 		for (int i = 0; i < BufferRing.Num(); ++i)
 		{
-			ProduceIndex = (ProduceIndex + i) % BufferRing.Num();
+			ProduceIndex = (ProduceIndex + 1) % BufferRing.Num();
 			if (ProduceIndex == ConsumeIndex)
+			{
 				continue;
+			}
 			if (BufferRing[ProduceIndex].IsUnique())
 			{
-				Full = false;
+				bIsFull = false;
 				break;
 			}
 		}
 
-		if (Full)
+		if (bIsFull)
 		{
 			if (BufferRing.Num() < MaxSize)
 			{
