[PS2] Fix: Hang during first decoded frame

diff --git a/Engine/Plugins/Media/PixelCapture/Source/PixelCapture/Private/PixelCaptureCapturerMultiFormat.cpp b/Engine/Plugins/Media/PixelCapture/Source/PixelCapture/Private/PixelCaptureCapturerMultiFormat.cpp
--- a/Engine/Plugins/Media/PixelCapture/Source/PixelCapture/Private/PixelCaptureCapturerMultiFormat.cpp
+++ b/Engine/Plugins/Media/PixelCapture/Source/PixelCapture/Private/PixelCaptureCapturerMultiFormat.cpp
@@ -62,7 +62,12 @@ void FPixelCaptureCapturerMultiFormat::AddOutputFormat(int32 Format)
 
 	if (!FormatCapturers.Contains(Format))
 	{
-		TSharedPtr<FPixelCaptureCapturerLayered> NewCapturer = FPixelCaptureCapturerLayered::Create(CapturerSource, Format, LayerSizes.Array());
+		TArray<FIntPoint> LayersArray;
+		{
+			FScopeLock LayerLock(&LayersGuard);
+			LayersArray = LayerSizes.Array();
+		}
+		TSharedPtr<FPixelCaptureCapturerLayered> NewCapturer = FPixelCaptureCapturerLayered::Create(CapturerSource, Format, LayersArray);
 		NewCapturer->OnComplete.AddSP(this, &FPixelCaptureCapturerMultiFormat::OnCaptureFormatComplete, Format);
 		FormatCapturers.Add(Format, NewCapturer);
 	}
@@ -79,10 +84,13 @@ TSharedPtr<IPixelCaptureOutputFrame> FPixelCaptureCapturerMultiFormat::RequestFo
 		}
 	}
 
-	LayerSizes.Add(LayerSize);
-	LayerSizes.Sort([](const FIntPoint& A, const FIntPoint& B) {
-		return A.X * A.Y < B.X * B.Y;
-	});
+	{
+		FScopeLock LayerLock(&LayersGuard);
+		LayerSizes.Add(LayerSize);
+		LayerSizes.Sort([](const FIntPoint& A, const FIntPoint& B) {
+			return A.X * A.Y < B.X * B.Y;
+		});
+	}
 
 	// if we reached here then we dont have a pipeline for the given format.
 	// add it and return null
diff --git a/Engine/Plugins/Media/PixelCapture/Source/PixelCapture/Public/PixelCaptureCapturerMultiFormat.h b/Engine/Plugins/Media/PixelCapture/Source/PixelCapture/Public/PixelCaptureCapturerMultiFormat.h
--- a/Engine/Plugins/Media/PixelCapture/Source/PixelCapture/Public/PixelCaptureCapturerMultiFormat.h
+++ b/Engine/Plugins/Media/PixelCapture/Source/PixelCapture/Public/PixelCaptureCapturerMultiFormat.h
@@ -140,6 +140,7 @@ protected:
 
 	IPixelCaptureCapturerSource* CapturerSource;
 
+	mutable FCriticalSection LayersGuard;
 	TSet<FIntPoint> LayerSizes;
 
 	TMap<int32, TSharedPtr<FPixelCaptureCapturerLayered>> FormatCapturers;
diff --git a/Engine/Plugins/Media/PixelStreaming2/Source/PixelStreaming2/Internal/VideoCapturer.h b/Engine/Plugins/Media/PixelStreaming2/Source/PixelStreaming2/Internal/VideoCapturer.h
--- a/Engine/Plugins/Media/PixelStreaming2/Source/PixelStreaming2/Internal/VideoCapturer.h
+++ b/Engine/Plugins/Media/PixelStreaming2/Source/PixelStreaming2/Internal/VideoCapturer.h
@@ -28,6 +28,8 @@ namespace UE::PixelStreaming2
 
 		TSharedPtr<FVideoProducer> GetVideoProducer() { return VideoProducer; }
 
+		TSharedPtr<IPixelCaptureOutputFrame> WaitForFormat(int32 Format, TOptional<FIntPoint> Resolution = TOptional<FIntPoint>(), uint32 MaxWaitTime = 5000);
+
 		TSharedPtr<IPixelCaptureOutputFrame> RequestFormat(int32 Format, TOptional<FIntPoint> Resolution = TOptional<FIntPoint>());
 
 		void ResetFrameCapturer();
diff --git a/Engine/Plugins/Media/PixelStreaming2/Source/PixelStreaming2/Private/VideoCapturer.cpp b/Engine/Plugins/Media/PixelStreaming2/Source/PixelStreaming2/Private/VideoCapturer.cpp
--- a/Engine/Plugins/Media/PixelStreaming2/Source/PixelStreaming2/Private/VideoCapturer.cpp
+++ b/Engine/Plugins/Media/PixelStreaming2/Source/PixelStreaming2/Private/VideoCapturer.cpp
@@ -107,6 +107,19 @@ namespace UE::PixelStreaming2
 		FrameCapturer->Capture(InputFrame);
 	}
 
+	TSharedPtr<IPixelCaptureOutputFrame> FVideoCapturer::WaitForFormat(int32 Format, TOptional<FIntPoint> Resolution, uint32 MaxWaitTime)
+	{
+		if (FrameCapturer != nullptr)
+		{
+			if (!Resolution.IsSet())
+			{
+				Resolution = { LastFrameWidth, LastFrameHeight };
+			}
+			return FrameCapturer->WaitForFormat(Format, *Resolution, MaxWaitTime);
+		}
+		return nullptr;
+	}
+
 	TSharedPtr<IPixelCaptureOutputFrame> FVideoCapturer::RequestFormat(int32 Format, TOptional<FIntPoint> Resolution)
 	{
 		if (FrameCapturer != nullptr)
@@ -115,7 +128,7 @@ namespace UE::PixelStreaming2
 			{
 				Resolution = { LastFrameWidth, LastFrameHeight };
 			}
-			return FrameCapturer->WaitForFormat(Format, *Resolution);
+			return FrameCapturer->RequestFormat(Format, *Resolution);
 		}
 		return nullptr;
 	}
diff --git a/Engine/Plugins/Media/PixelStreaming2/Source/PixelStreaming2RTC/Private/EpicRtcVideoBufferMultiFormat.cpp b/Engine/Plugins/Media/PixelStreaming2/Source/PixelStreaming2RTC/Private/EpicRtcVideoBufferMultiFormat.cpp
--- a/Engine/Plugins/Media/PixelStreaming2/Source/PixelStreaming2RTC/Private/EpicRtcVideoBufferMultiFormat.cpp
+++ b/Engine/Plugins/Media/PixelStreaming2/Source/PixelStreaming2RTC/Private/EpicRtcVideoBufferMultiFormat.cpp
@@ -58,8 +58,12 @@ namespace UE::PixelStreaming2
 		{
 			return nullptr;
 		}
-		constexpr uint32					 WaitTimeMS = 500;
-		TSharedPtr<IPixelCaptureOutputFrame> Frame = FrameCapturer->WaitForFormat(Format, Resolution, WaitTimeMS);
+		TSharedPtr<IPixelCaptureOutputFrame> Frame = FrameCapturer->RequestFormat(Format, Resolution);
+		if (!Frame) 
+		{
+			return nullptr;
+		}
+		
 		CachedFormat.Add(Format, Frame);
 		return Frame.Get();
 	}
diff --git a/Engine/Plugins/Media/PixelStreaming2/Source/PixelStreaming2RTC/Private/FreezeFrame.cpp b/Engine/Plugins/Media/PixelStreaming2/Source/PixelStreaming2RTC/Private/FreezeFrame.cpp
--- a/Engine/Plugins/Media/PixelStreaming2/Source/PixelStreaming2RTC/Private/FreezeFrame.cpp
+++ b/Engine/Plugins/Media/PixelStreaming2/Source/PixelStreaming2RTC/Private/FreezeFrame.cpp
@@ -276,7 +276,7 @@ namespace UE::PixelStreaming2
 		if (CurrentCodec == EVideoCodec::VP8 || CurrentCodec == EVideoCodec::VP9)
 		{
 			// Request output format is I420 for VPX
-			TSharedPtr<IPixelCaptureOutputFrame> OutputFrame = Input->RequestFormat(PixelCaptureBufferFormat::FORMAT_I420);
+			TSharedPtr<IPixelCaptureOutputFrame> OutputFrame = Input->WaitForFormat(PixelCaptureBufferFormat::FORMAT_I420);
 			if (OutputFrame)
 			{
 				// Can remove binding now we have got the output in the format we need to send a FF
@@ -303,7 +303,7 @@ namespace UE::PixelStreaming2
 		}
 		else
 		{
-			TSharedPtr<IPixelCaptureOutputFrame> OutputFrame = Input->RequestFormat(PixelCaptureBufferFormat::FORMAT_RHI);
+			TSharedPtr<IPixelCaptureOutputFrame> OutputFrame = Input->WaitForFormat(PixelCaptureBufferFormat::FORMAT_RHI);
 			if (OutputFrame)
 			{
 				// Can remove binding now we have got the output in the format we need to send a FF
