[PixelCapture, PS, PS2] Fix: MediaCapture could get into a bad state due to use of queues and praying frame timing lined up

diff --git a/Engine/Plugins/Media/PixelCapture/Source/PixelCapture/Internal/PixelCaptureCapturerMediaCapture.h b/Engine/Plugins/Media/PixelCapture/Source/PixelCapture/Internal/PixelCaptureCapturerMediaCapture.h
--- a/Engine/Plugins/Media/PixelCapture/Source/PixelCapture/Internal/PixelCaptureCapturerMediaCapture.h
+++ b/Engine/Plugins/Media/PixelCapture/Source/PixelCapture/Internal/PixelCaptureCapturerMediaCapture.h
@@ -2,7 +2,7 @@
 
 #pragma once
 
-#include "Containers/Deque.h"
+#include "Containers/StripedMap.h"
 #include "MediaCapture.h"
 #include "MediaOutput.h"
 #include "PixelCaptureBufferFormat.h"
@@ -33,14 +33,14 @@ class UPixelCaptureMediaCapture : public UMediaCapture
 	GENERATED_BODY()
 
 public:
-	void AddOutputFrame(TSharedPtr<IPixelCaptureOutputFrame> InOutputFrame)
+	void AddOutputFrame(uint64 InFrameNumberRenderThread, TSharedPtr<IPixelCaptureOutputFrame> InOutputFrame)
 	{
-		OutputFrames.PushLast(InOutputFrame);
+		OutputFrames.Add(InFrameNumberRenderThread, InOutputFrame);
 	}
 
-	void RemoveLastOutputFrame()
+	void RemoveOutputFrame(uint64 InFrameNumberRenderThread)
 	{
-		OutputFrames.PopLast();
+		OutputFrames.Remove(InFrameNumberRenderThread);
 	}
 
 	void SetFormat(int32 InFormat)
@@ -116,15 +116,15 @@ protected:
 
 	UE_API virtual ETextureCreateFlags GetOutputTextureFlags() const override;
 	UE_API virtual void				WaitForGPU(FRHITexture* InRHITexture) override;
-
 private:
 	int32 Format = PixelCaptureBufferFormat::FORMAT_UNKNOWN;
 
 	bool bIsSRGB = false;
 
-	TDeque<FGPUFenceRHIRef> Fences;
+	TStripedMap<32, uintptr_t, FGPUFenceRHIRef> Fences;
 
-	TDeque<TSharedPtr<IPixelCaptureOutputFrame>> OutputFrames;
+	//  SourceFrameNumberRenderThread, OutputFrame
+	TStripedMap<32, uint64, TSharedPtr<IPixelCaptureOutputFrame>> OutputFrames;
 };
 
 /**
@@ -143,7 +143,7 @@ public:
 	UE_API virtual ~FPixelCaptureCapturerMediaCapture() override;
 
 protected:
-	virtual FString					  GetCapturerName() const override { return "MediaCapture Copy"; }
+	UE_API virtual FString					  GetCapturerName() const override;
 	UE_API virtual IPixelCaptureOutputFrame* CreateOutputBuffer(int32 InputWidth, int32 InputHeight) override;
 	UE_API virtual void					  BeginProcess(const IPixelCaptureInputFrame& InputFrame, TSharedPtr<IPixelCaptureOutputFrame> OutputBuffer) override;
 
diff --git a/Engine/Plugins/Media/PixelCapture/Source/PixelCapture/Private/PixelCaptureCapturerMediaCapture.cpp b/Engine/Plugins/Media/PixelCapture/Source/PixelCapture/Private/PixelCaptureCapturerMediaCapture.cpp
--- a/Engine/Plugins/Media/PixelCapture/Source/PixelCapture/Private/PixelCaptureCapturerMediaCapture.cpp
+++ b/Engine/Plugins/Media/PixelCapture/Source/PixelCapture/Private/PixelCaptureCapturerMediaCapture.cpp
@@ -19,15 +19,15 @@ void UPixelCaptureMediaCapture::OnRHIResourceCaptured_AnyThread(
 	TSharedPtr<FMediaCaptureUserData, ESPMode::ThreadSafe> UserData,
 	FTextureRHIRef										   Texture)
 {
+	check(Format == PixelCaptureBufferFormat::FORMAT_RHI);
+
 	TSharedPtr<IPixelCaptureOutputFrame> OutputFrame;
-	if (!OutputFrames.TryPopFirst(OutputFrame) || !OutputFrame)
+	if (!OutputFrames.RemoveAndCopyValue(BaseData.SourceFrameNumberRenderThread, OutputFrame) || !OutputFrame)
 	{
-		UE_LOGFMT(LogPixelCapture, Warning, "UPixelCaptureMediaCapture::OnRHIResourceCaptured_AnyThread: No output frame set!");
+		UE_LOGFMT(LogPixelCapture, Warning, "UPixelCaptureMediaCapture::OnRHIResourceCaptured_AnyThread: No output frame set for frame number {0}!", BaseData.SourceFrameNumberRenderThread);
 		return;
 	}
 
-	check(Format == PixelCaptureBufferFormat::FORMAT_RHI);
-
 	StaticCastSharedPtr<FPixelCaptureOutputFrameRHI>(OutputFrame)->SetFrameTexture(Texture);
 
 	OnCaptureComplete.Broadcast(OutputFrame);
@@ -47,15 +47,15 @@ void UPixelCaptureMediaCapture::OnFrameCaptured_AnyThread(
 	TSharedPtr<FMediaCaptureUserData, ESPMode::ThreadSafe> UserData,
 	const FMediaCaptureResourceData&					   ResourceData)
 {
+	check(Format == PixelCaptureBufferFormat::FORMAT_I420);
+
 	TSharedPtr<IPixelCaptureOutputFrame> OutputFrame;
-	if (!OutputFrames.TryPopFirst(OutputFrame) || !OutputFrame)
+	if (!OutputFrames.RemoveAndCopyValue(BaseData.SourceFrameNumberRenderThread, OutputFrame) || !OutputFrame)
 	{
-		UE_LOGFMT(LogPixelCapture, Warning, "UPixelCaptureMediaCapture::OnFrameCaptured_AnyThread: No output frame set!");
+		UE_LOGFMT(LogPixelCapture, Warning, "UPixelCaptureMediaCapture::OnFrameCaptured_AnyThread: No output frame set for frame number {0}!", BaseData.SourceFrameNumberRenderThread);
 		return;
 	}
 
-	check(Format == PixelCaptureBufferFormat::FORMAT_I420);
-
 	TSharedPtr<FPixelCaptureBufferI420> I420Buffer = MakeShared<FPixelCaptureBufferI420>(ResourceData.Width, ResourceData.Height);
 	libyuv::ARGBToI420(
 		static_cast<uint8*>(ResourceData.Buffer),
@@ -99,7 +99,7 @@ void UPixelCaptureMediaCapture::OnCustomCapture_RenderingThread(
 	FGPUFenceRHIRef Fence = GDynamicRHI->RHICreateGPUFence(TEXT("UPixelCaptureMediaCapture Fence"));
 	if (UseExperimentalScheduling() && ShouldCaptureRHIResource())
 	{
-		Fences.PushLast(Fence);
+		Fences.Add(reinterpret_cast<uintptr_t>(OutputTexture->GetRHI()), Fence);
 	}
 
 	bool bRequiresFormatConversion = InSourceTexture->Desc.Format != OutputTexture->Desc.Format;
@@ -215,8 +215,9 @@ ETextureCreateFlags UPixelCaptureMediaCapture::GetOutputTextureFlags() const
 void UPixelCaptureMediaCapture::WaitForGPU(FRHITexture* InRHITexture)
 {
 	FGPUFenceRHIRef Fence;
-	if (!Fences.TryPopFirst(Fence))
+	if (!Fences.RemoveAndCopyValue(reinterpret_cast<uintptr_t>(InRHITexture), Fence))
 	{
+		UE_LOGFMT(LogPixelCapture, Error, "UPixelCaptureMediaCapture::WaitForGPU: No associating fence!");
 		return;
 	}
 
@@ -259,6 +260,7 @@ FPixelCaptureCapturerMediaCapture::~FPixelCaptureCapturerMediaCapture()
 	// as UE will already have killed all UObjects by this point.
 	if (!IsEngineExitRequested() && MediaCapture)
 	{
+		MediaCapture->StopCapture(false);
 		MediaCapture->RemoveFromRoot();
 	}
 }
@@ -298,6 +300,22 @@ void FPixelCaptureCapturerMediaCapture::InitializeMediaCapture()
 	bMediaCaptureInitialized = true;
 }
 
+FString FPixelCaptureCapturerMediaCapture::GetCapturerName() const 
+{
+	FString FormatString = TEXT("Unknown");
+	switch (Format)
+	{
+		case PixelCaptureBufferFormat::FORMAT_RHI:
+			FormatString = TEXT("RHI");
+			break;
+		case PixelCaptureBufferFormat::FORMAT_I420:
+			FormatString = TEXT("I420");
+			break;
+	}
+	
+	return FString::Printf(TEXT("MediaCapture (RHI->%s)"), *FormatString);
+}
+
 IPixelCaptureOutputFrame* FPixelCaptureCapturerMediaCapture::CreateOutputBuffer(int32 InputWidth, int32 InputHeight)
 {
 	if (Format == PixelCaptureBufferFormat::FORMAT_RHI)
@@ -351,7 +369,7 @@ void FPixelCaptureCapturerMediaCapture::BeginProcess(const IPixelCaptureInputFra
 	UE::PixelCapture::MarkCPUWorkStart(OutputBuffer);
 
 	// We need to ensure the output frame is valid before calling TryCaptureImmediate_RenderThread
-	MediaCapture->AddOutputFrame(OutputBuffer);
+	MediaCapture->AddOutputFrame(LastFrameCounterRenderThread, OutputBuffer);
 
 	FRDGBuilder GraphBuilder(FRHICommandListImmediate::Get());
 
@@ -383,7 +401,7 @@ void FPixelCaptureCapturerMediaCapture::BeginProcess(const IPixelCaptureInputFra
 	if (!bPassesAdded)
 	{
 		// TryCaptureImmediate_RenderThread didn't add any passes so remove the latest enqueued output frame from the MediaCapture
-		MediaCapture->RemoveLastOutputFrame();
+		MediaCapture->RemoveOutputFrame(LastFrameCounterRenderThread);
 		// No passes added. Invalidate the output buffer so that the encoder won't pull old data from the ring buffer
 		InvalidateOutputBuffer(OutputBuffer);
 		// RDG graph had no passes so we can manually call EndProcess()
diff --git a/Engine/Plugins/Media/PixelStreaming/Source/PixelStreaming/Internal/PixelStreamingMediaIOCapture.h b/Engine/Plugins/Media/PixelStreaming/Source/PixelStreaming/Internal/PixelStreamingMediaIOCapture.h
--- a/Engine/Plugins/Media/PixelStreaming/Source/PixelStreaming/Internal/PixelStreamingMediaIOCapture.h
+++ b/Engine/Plugins/Media/PixelStreaming/Source/PixelStreaming/Internal/PixelStreamingMediaIOCapture.h
@@ -2,7 +2,7 @@
 
 #pragma once
 
-#include "Containers/Deque.h"
+#include "Containers/StripedMap.h"
 #include "MediaCapture.h"
 #include "PixelStreamingVideoInput.h"
 #include "Slate/SceneViewport.h"
@@ -90,6 +90,6 @@ private:
 	/* Whether we want the UMediaCapture to read back to frame into cpu memory or not. */
 	bool bDoGPUCopy = true;
 
-	// Deque of fences used to wait on to ensure the media capture conversion is complete
-	TDeque<FGPUFenceRHIRef> Fences;
+	// Thread safe map of fences used to wait on to ensure the media capture conversion is complete
+	TStripedMap<32, uintptr_t, FGPUFenceRHIRef> Fences;
 };
diff --git a/Engine/Plugins/Media/PixelStreaming/Source/PixelStreaming/Private/PixelStreamingMediaIOCapture.cpp b/Engine/Plugins/Media/PixelStreaming/Source/PixelStreaming/Private/PixelStreamingMediaIOCapture.cpp
--- a/Engine/Plugins/Media/PixelStreaming/Source/PixelStreaming/Private/PixelStreamingMediaIOCapture.cpp
+++ b/Engine/Plugins/Media/PixelStreaming/Source/PixelStreaming/Private/PixelStreamingMediaIOCapture.cpp
@@ -56,7 +56,7 @@ void UPixelStreamingMediaIOCapture::OnCustomCapture_RenderingThread(
 	FGPUFenceRHIRef Fence = GDynamicRHI->RHICreateGPUFence(TEXT("UPixelStreamingMediaIOCapture Fence"));
 	if (UseExperimentalScheduling() && ShouldCaptureRHIResource())
 	{
-		Fences.PushLast(Fence);
+		Fences.Add(reinterpret_cast<uintptr_t>(OutputTexture->GetRHI()), Fence);
 	}
 
 	bool bRequiresFormatConversion = InSourceTexture->Desc.Format != OutputTexture->Desc.Format;
@@ -190,8 +190,9 @@ ETextureCreateFlags UPixelStreamingMediaIOCapture::GetOutputTextureFlags() const
 void UPixelStreamingMediaIOCapture::WaitForGPU(FRHITexture* InRHITexture)
 {
 	FGPUFenceRHIRef Fence;
-	if (!Fences.TryPopFirst(Fence))
+	if (!Fences.RemoveAndCopyValue(reinterpret_cast<uintptr_t>(InRHITexture), Fence))
 	{
+		UE_LOG(LogPixelStreaming, Error, TEXT("UPixelStreamingMediaIOCapture::WaitForGPU: No associating fence!"));
 		return;
 	}
 
diff --git a/Engine/Plugins/Media/PixelStreaming2/Source/PixelStreaming2/Internal/PixelStreaming2MediaIOCapture.h b/Engine/Plugins/Media/PixelStreaming2/Source/PixelStreaming2/Internal/PixelStreaming2MediaIOCapture.h
--- a/Engine/Plugins/Media/PixelStreaming2/Source/PixelStreaming2/Internal/PixelStreaming2MediaIOCapture.h
+++ b/Engine/Plugins/Media/PixelStreaming2/Source/PixelStreaming2/Internal/PixelStreaming2MediaIOCapture.h
@@ -2,7 +2,7 @@
 
 #pragma once
 
-#include "Containers/Deque.h"
+#include "Containers/StripedMap.h"
 #include "MediaCapture.h"
 #include "Slate/SceneViewport.h"
 #include "VideoProducer.h"
@@ -76,7 +76,7 @@ private:
 	/* Whether we want the UMediaCapture to read back to frame into cpu memory or not. */
 	bool bDoGPUCopy = true;
 
-	TDeque<FGPUFenceRHIRef> Fences;
+	TStripedMap<32, uintptr_t, FGPUFenceRHIRef> Fences;
 };
 
 #undef UE_API
diff --git a/Engine/Plugins/Media/PixelStreaming2/Source/PixelStreaming2/Private/PixelStreaming2MediaIOCapture.cpp b/Engine/Plugins/Media/PixelStreaming2/Source/PixelStreaming2/Private/PixelStreaming2MediaIOCapture.cpp
--- a/Engine/Plugins/Media/PixelStreaming2/Source/PixelStreaming2/Private/PixelStreaming2MediaIOCapture.cpp
+++ b/Engine/Plugins/Media/PixelStreaming2/Source/PixelStreaming2/Private/PixelStreaming2MediaIOCapture.cpp
@@ -31,7 +31,7 @@ void UPixelStreaming2MediaIOCapture::OnCustomCapture_RenderingThread(
 	FGPUFenceRHIRef Fence = GDynamicRHI->RHICreateGPUFence(TEXT("UPixelStreaming2MediaIOCapture Fence"));
 	if (UseExperimentalScheduling() && ShouldCaptureRHIResource())
 	{
-		Fences.PushLast(Fence);
+		Fences.Add(reinterpret_cast<uintptr_t>(OutputTexture->GetRHI()), Fence);
 	}
 
 	bool bRequiresFormatConversion = InSourceTexture->Desc.Format != OutputTexture->Desc.Format;
@@ -158,8 +158,9 @@ ETextureCreateFlags UPixelStreaming2MediaIOCapture::GetOutputTextureFlags() cons
 void UPixelStreaming2MediaIOCapture::WaitForGPU(FRHITexture* InRHITexture)
 {
 	FGPUFenceRHIRef Fence;
-	if (!Fences.TryPopFirst(Fence))
+	if (!Fences.RemoveAndCopyValue(reinterpret_cast<uintptr_t>(InRHITexture), Fence))
 	{
+		UE_LOGFMT(LogPixelStreaming2, Error, "UPixelStreaming2MediaIOCapture::WaitForGPU: No associating fence!");
 		return;
 	}
 
