[PixelCapture] Fix: Remove call to StopCapture as it can cause deadlocks

diff --git a/Engine/Plugins/Media/PixelCapture/Source/PixelCapture/Private/PixelCaptureCapturerMediaCapture.cpp b/Engine/Plugins/Media/PixelCapture/Source/PixelCapture/Private/PixelCaptureCapturerMediaCapture.cpp
--- a/Engine/Plugins/Media/PixelCapture/Source/PixelCapture/Private/PixelCaptureCapturerMediaCapture.cpp
+++ b/Engine/Plugins/Media/PixelCapture/Source/PixelCapture/Private/PixelCaptureCapturerMediaCapture.cpp
@@ -260,7 +260,6 @@ FPixelCaptureCapturerMediaCapture::~FPixelCaptureCapturerMediaCapture()
 	// as UE will already have killed all UObjects by this point.
 	if (!IsEngineExitRequested() && MediaCapture)
 	{
-		MediaCapture->StopCapture(false);
 		MediaCapture->RemoveFromRoot();
 	}
 }
