PixelCapture code needs to pass a command list around instead of using implicit command lists.

diff --git a/Engine/Plugins/Media/PixelCapture/Source/PixelCapture/Internal/PixelCaptureCapturerMediaCapture.h b/Engine/Plugins/Media/PixelCapture/Source/PixelCapture/Internal/PixelCaptureCapturerMediaCapture.h
--- a/Engine/Plugins/Media/PixelCapture/Source/PixelCapture/Internal/PixelCaptureCapturerMediaCapture.h
+++ b/Engine/Plugins/Media/PixelCapture/Source/PixelCapture/Internal/PixelCaptureCapturerMediaCapture.h
@@ -144,8 +144,8 @@ public:
 
 protected:
 	UE_API virtual FString					  GetCapturerName() const override;
-	UE_API virtual IPixelCaptureOutputFrame* CreateOutputBuffer(int32 InputWidth, int32 InputHeight) override;
-	UE_API virtual void					  BeginProcess(const IPixelCaptureInputFrame& InputFrame, TSharedPtr<IPixelCaptureOutputFrame> OutputBuffer) override;
+	UE_API virtual IPixelCaptureOutputFrame* CreateOutputBuffer(FRHICommandListBase& RHICmdList, int32 InputWidth, int32 InputHeight) override;
+	UE_API virtual void                      BeginProcess(FRHICommandList& RHICmdList, const IPixelCaptureInputFrame& InputFrame, TSharedPtr<IPixelCaptureOutputFrame> OutputBuffer) override;
 
 private:
 	uint64 LastFrameCounterRenderThread = 0;
diff --git a/Engine/Plugins/Media/PixelCapture/Source/PixelCapture/Private/PixelCaptureCapturer.cpp b/Engine/Plugins/Media/PixelCapture/Source/PixelCapture/Private/PixelCaptureCapturer.cpp
--- a/Engine/Plugins/Media/PixelCapture/Source/PixelCapture/Private/PixelCaptureCapturer.cpp
+++ b/Engine/Plugins/Media/PixelCapture/Source/PixelCapture/Private/PixelCaptureCapturer.cpp
@@ -3,6 +3,7 @@
 #include "PixelCaptureCapturer.h"
 #include "OutputFrameBuffer.h"
 #include "PixelCapturePrivate.h"
+#include "RHICommandList.h"
 
 namespace UE::PixelCapture
 {
@@ -92,7 +93,7 @@ TSharedPtr<IPixelCaptureOutputFrame> FPixelCaptureCapturer::ReadOutput()
 	return nullptr;
 }
 
-void FPixelCaptureCapturer::Capture(const IPixelCaptureInputFrame& InputFrame)
+void FPixelCaptureCapturer::Capture(FRHICommandList& RHICmdList, const IPixelCaptureInputFrame& InputFrame)
 {
 	if (IsBusy())
 	{
@@ -101,7 +102,7 @@ void FPixelCaptureCapturer::Capture(const IPixelCaptureInputFrame& InputFrame)
 
 	if (!IsInitialized())
 	{
-		Initialize(InputFrame.GetWidth(), InputFrame.GetHeight());
+		Initialize(RHICmdList, InputFrame.GetWidth(), InputFrame.GetHeight());
 	}
 
 	checkf(InputFrame.GetWidth() == ExpectedInputWidth && InputFrame.GetHeight() == ExpectedInputHeight, TEXT("Capturer input resolution changes are not supported"));
@@ -115,13 +116,13 @@ void FPixelCaptureCapturer::Capture(const IPixelCaptureInputFrame& InputFrame)
 
 	InitMetadata(InputFrame.Metadata.Copy(), OutputBuffer);
 
-	BeginProcess(InputFrame, OutputBuffer);
+	BeginProcess(RHICmdList, InputFrame, OutputBuffer);
 }
 
-void FPixelCaptureCapturer::Initialize(int32 InputWidth, int32 InputHeight)
+void FPixelCaptureCapturer::Initialize(FRHICommandListBase& RHICmdList, int32 InputWidth, int32 InputHeight)
 {
 	Buffer = MakeUnique<UE::PixelCapture::FOutputFrameBuffer>();
-	Buffer->Reset(3, 32, [this, InputWidth, InputHeight]() { return TSharedPtr<IPixelCaptureOutputFrame>(CreateOutputBuffer(InputWidth, InputHeight)); });
+	Buffer->Reset(3, 32, [this, &RHICmdList, InputWidth, InputHeight]() { return TSharedPtr<IPixelCaptureOutputFrame>(CreateOutputBuffer(RHICmdList, InputWidth, InputHeight)); });
 	ExpectedInputWidth = InputWidth;
 	ExpectedInputHeight = InputHeight;
 	bHasOutput = false;
diff --git a/Engine/Plugins/Media/PixelCapture/Source/PixelCapture/Private/PixelCaptureCapturerI420.cpp b/Engine/Plugins/Media/PixelCapture/Source/PixelCapture/Private/PixelCaptureCapturerI420.cpp
--- a/Engine/Plugins/Media/PixelCapture/Source/PixelCapture/Private/PixelCaptureCapturerI420.cpp
+++ b/Engine/Plugins/Media/PixelCapture/Source/PixelCapture/Private/PixelCaptureCapturerI420.cpp
@@ -15,17 +15,17 @@ FPixelCaptureCapturerI420::FPixelCaptureCapturerI420(FPixelCaptureCapturerConfig
 {
 }
 
-void FPixelCaptureCapturerI420::Initialize(int32 InputWidth, int32 InputHeight)
+void FPixelCaptureCapturerI420::Initialize(FRHICommandListBase& RHICmdList, int32 InputWidth, int32 InputHeight)
 {
-	FPixelCaptureCapturer::Initialize(InputWidth, InputHeight);
+	FPixelCaptureCapturer::Initialize(RHICmdList, InputWidth, InputHeight);
 }
 
-IPixelCaptureOutputFrame* FPixelCaptureCapturerI420::CreateOutputBuffer(int32 InputWidth, int32 InputHeight)
+IPixelCaptureOutputFrame* FPixelCaptureCapturerI420::CreateOutputBuffer(FRHICommandListBase& RHICmdList, int32 InputWidth, int32 InputHeight)
 {
 	return new FPixelCaptureOutputFrameI420(MakeShared<FPixelCaptureBufferI420>(InputWidth, InputHeight));
 }
 
-void FPixelCaptureCapturerI420::BeginProcess(const IPixelCaptureInputFrame& InputFrame, TSharedPtr<IPixelCaptureOutputFrame> OutputBuffer)
+void FPixelCaptureCapturerI420::BeginProcess(FRHICommandList& RHICmdList, const IPixelCaptureInputFrame& InputFrame, TSharedPtr<IPixelCaptureOutputFrame> OutputBuffer)
 {
 	SetIsBusy(true);
 
diff --git a/Engine/Plugins/Media/PixelCapture/Source/PixelCapture/Private/PixelCaptureCapturerI420ToRHI.cpp b/Engine/Plugins/Media/PixelCapture/Source/PixelCapture/Private/PixelCaptureCapturerI420ToRHI.cpp
--- a/Engine/Plugins/Media/PixelCapture/Source/PixelCapture/Private/PixelCaptureCapturerI420ToRHI.cpp
+++ b/Engine/Plugins/Media/PixelCapture/Source/PixelCapture/Private/PixelCaptureCapturerI420ToRHI.cpp
@@ -33,17 +33,17 @@ FPixelCaptureCapturerI420ToRHI::~FPixelCaptureCapturerI420ToRHI()
 	}
 }
 
-void FPixelCaptureCapturerI420ToRHI::Initialize(int32 InputWidth, int32 InputHeight)
+void FPixelCaptureCapturerI420ToRHI::Initialize(FRHICommandListBase& RHICmdList, int32 InputWidth, int32 InputHeight)
 {
 	Config.OutputResolution = { InputWidth, InputHeight };
 
 	const uint32_t Size = webrtc::CalcBufferSize(webrtc::VideoType::kARGB, Config.OutputResolution.X, Config.OutputResolution.Y);
 	ARGBBuffer = new uint8_t[Size];
 
-	FPixelCaptureCapturer::Initialize(InputWidth, InputHeight);
+	FPixelCaptureCapturer::Initialize(RHICmdList, InputWidth, InputHeight);
 }
 
-IPixelCaptureOutputFrame* FPixelCaptureCapturerI420ToRHI::CreateOutputBuffer(int32 InputWidth, int32 InputHeight)
+IPixelCaptureOutputFrame* FPixelCaptureCapturerI420ToRHI::CreateOutputBuffer(FRHICommandListBase& RHICmdList, int32 InputWidth, int32 InputHeight)
 {
 	FRHITextureCreateDesc TextureDesc =
 		FRHITextureCreateDesc::Create2D(TEXT("FPixelCaptureCapturerI420ToRHI Texture"), InputWidth, InputHeight, EPixelFormat::PF_B8G8R8A8)
@@ -66,10 +66,10 @@ IPixelCaptureOutputFrame* FPixelCaptureCapturerI420ToRHI::CreateOutputBuffer(int
 		TextureDesc.AddFlags(ETextureCreateFlags::Shared);
 	}
 
-	return new FPixelCaptureOutputFrameRHI(RHICreateTexture(TextureDesc));
+	return new FPixelCaptureOutputFrameRHI(RHICmdList.CreateTexture(TextureDesc));
 }
 
-void FPixelCaptureCapturerI420ToRHI::BeginProcess(const IPixelCaptureInputFrame& InputFrame, TSharedPtr<IPixelCaptureOutputFrame> OutputBuffer)
+void FPixelCaptureCapturerI420ToRHI::BeginProcess(FRHICommandList& RHICmdList, const IPixelCaptureInputFrame& InputFrame, TSharedPtr<IPixelCaptureOutputFrame> OutputBuffer)
 {
 	SetIsBusy(true);
 
diff --git a/Engine/Plugins/Media/PixelCapture/Source/PixelCapture/Private/PixelCaptureCapturerLayered.cpp b/Engine/Plugins/Media/PixelCapture/Source/PixelCapture/Private/PixelCaptureCapturerLayered.cpp
--- a/Engine/Plugins/Media/PixelCapture/Source/PixelCapture/Private/PixelCaptureCapturerLayered.cpp
+++ b/Engine/Plugins/Media/PixelCapture/Source/PixelCapture/Private/PixelCaptureCapturerLayered.cpp
@@ -51,7 +51,7 @@ void FPixelCaptureCapturerLayered::OnCaptureComplete()
 	OnComplete.Broadcast();
 }
 
-void FPixelCaptureCapturerLayered::Capture(const IPixelCaptureInputFrame& SourceFrame)
+void FPixelCaptureCapturerLayered::Capture(FRHICommandList& RHICmdList, const IPixelCaptureInputFrame& SourceFrame)
 {
 	// work on a temp list so we dont over lock
 	TArray<TSharedPtr<FPixelCaptureCapturer>> TempLayerCapturers;
@@ -61,8 +61,8 @@ void FPixelCaptureCapturerLayered::Capture(const IPixelCaptureInputFrame& Source
 	}
 
 	// capture the frame for encoder use
-	for (auto& LayerCapturer : TempLayerCapturers)
+	for (TSharedPtr<FPixelCaptureCapturer>& LayerCapturer : TempLayerCapturers)
 	{
-		LayerCapturer->Capture(SourceFrame);
+		LayerCapturer->Capture(RHICmdList, SourceFrame);
 	}
 }
diff --git a/Engine/Plugins/Media/PixelCapture/Source/PixelCapture/Private/PixelCaptureCapturerMediaCapture.cpp b/Engine/Plugins/Media/PixelCapture/Source/PixelCapture/Private/PixelCaptureCapturerMediaCapture.cpp
--- a/Engine/Plugins/Media/PixelCapture/Source/PixelCapture/Private/PixelCaptureCapturerMediaCapture.cpp
+++ b/Engine/Plugins/Media/PixelCapture/Source/PixelCapture/Private/PixelCaptureCapturerMediaCapture.cpp
@@ -316,7 +316,7 @@ FString FPixelCaptureCapturerMediaCapture::GetCapturerName() const
 	return FString::Printf(TEXT("MediaCapture (RHI->%s)"), *FormatString);
 }
 
-IPixelCaptureOutputFrame* FPixelCaptureCapturerMediaCapture::CreateOutputBuffer(int32 InputWidth, int32 InputHeight)
+IPixelCaptureOutputFrame* FPixelCaptureCapturerMediaCapture::CreateOutputBuffer(FRHICommandListBase& RHICmdList, int32 InputWidth, int32 InputHeight)
 {
 	if (Format == PixelCaptureBufferFormat::FORMAT_RHI)
 	{
@@ -333,7 +333,7 @@ IPixelCaptureOutputFrame* FPixelCaptureCapturerMediaCapture::CreateOutputBuffer(
 	}
 }
 
-void FPixelCaptureCapturerMediaCapture::BeginProcess(const IPixelCaptureInputFrame& InputFrame, TSharedPtr<IPixelCaptureOutputFrame> OutputBuffer)
+void FPixelCaptureCapturerMediaCapture::BeginProcess(FRHICommandList& RHICmdList, const IPixelCaptureInputFrame& InputFrame, TSharedPtr<IPixelCaptureOutputFrame> OutputBuffer)
 {
 	if (!bMediaCaptureInitialized)
 	{
@@ -371,7 +371,7 @@ void FPixelCaptureCapturerMediaCapture::BeginProcess(const IPixelCaptureInputFra
 	// We need to ensure the output frame is valid before calling TryCaptureImmediate_RenderThread
 	MediaCapture->AddOutputFrame(LastFrameCounterRenderThread, OutputBuffer);
 
-	FRDGBuilder GraphBuilder(FRHICommandListImmediate::Get());
+	FRDGBuilder GraphBuilder(RHICmdList.GetAsImmediate());
 
 	// clang-format off
 	GraphBuilder.AddPass(
diff --git a/Engine/Plugins/Media/PixelCapture/Source/PixelCapture/Private/PixelCaptureCapturerMultiFormat.cpp b/Engine/Plugins/Media/PixelCapture/Source/PixelCapture/Private/PixelCaptureCapturerMultiFormat.cpp
--- a/Engine/Plugins/Media/PixelCapture/Source/PixelCapture/Private/PixelCaptureCapturerMultiFormat.cpp
+++ b/Engine/Plugins/Media/PixelCapture/Source/PixelCapture/Private/PixelCaptureCapturerMultiFormat.cpp
@@ -28,7 +28,7 @@ FPixelCaptureCapturerMultiFormat::~FPixelCaptureCapturerMultiFormat()
 	FlushWaitingEvents();
 }
 
-void FPixelCaptureCapturerMultiFormat::Capture(const IPixelCaptureInputFrame& SourceFrame)
+void FPixelCaptureCapturerMultiFormat::Capture(FRHICommandList& RHICmdList, const IPixelCaptureInputFrame& SourceFrame)
 {
 	if (FormatCapturers.Num() > 0)
 	{
@@ -45,9 +45,9 @@ void FPixelCaptureCapturerMultiFormat::Capture(const IPixelCaptureInputFrame& So
 			AllFormatCapturers.Add(FormatCapturer.Value);
 		}
 		// start all the format captures
-		for (auto& Capturer : AllFormatCapturers)
+		for (TSharedPtr<FPixelCaptureCapturerLayered>& Capturer : AllFormatCapturers)
 		{
-			Capturer->Capture(SourceFrame);
+			Capturer->Capture(RHICmdList, SourceFrame);
 		}
 	}
 	else
diff --git a/Engine/Plugins/Media/PixelCapture/Source/PixelCapture/Private/PixelCaptureCapturerNV12ToRHI.cpp b/Engine/Plugins/Media/PixelCapture/Source/PixelCapture/Private/PixelCaptureCapturerNV12ToRHI.cpp
--- a/Engine/Plugins/Media/PixelCapture/Source/PixelCapture/Private/PixelCaptureCapturerNV12ToRHI.cpp
+++ b/Engine/Plugins/Media/PixelCapture/Source/PixelCapture/Private/PixelCaptureCapturerNV12ToRHI.cpp
@@ -32,17 +32,17 @@ FPixelCaptureCapturerNV12ToRHI::~FPixelCaptureCapturerNV12ToRHI()
 	}
 }
 
-void FPixelCaptureCapturerNV12ToRHI::Initialize(int32 InputWidth, int32 InputHeight)
+void FPixelCaptureCapturerNV12ToRHI::Initialize(FRHICommandListBase& RHICmdList, int32 InputWidth, int32 InputHeight)
 {
 	Config.OutputResolution = { InputWidth, InputHeight };
 
 	const uint32_t Size = (3 * Config.OutputResolution.X * Config.OutputResolution.Y) >> 1;
 	R8Buffer = new uint8_t[Size];
 
-	FPixelCaptureCapturer::Initialize(InputWidth, InputHeight);
+	FPixelCaptureCapturer::Initialize(RHICmdList, InputWidth, InputHeight);
 }
 
-IPixelCaptureOutputFrame* FPixelCaptureCapturerNV12ToRHI::CreateOutputBuffer(int32 InputWidth, int32 InputHeight)
+IPixelCaptureOutputFrame* FPixelCaptureCapturerNV12ToRHI::CreateOutputBuffer(FRHICommandListBase& RHICmdList, int32 InputWidth, int32 InputHeight)
 {
 	FRHITextureCreateDesc TextureDesc =
 		FRHITextureCreateDesc::Create2D(TEXT("FPixelCaptureCapturerNV12ToRHI Texture"), InputWidth, 3 * InputHeight / 2, EPixelFormat::PF_R8)
@@ -65,10 +65,10 @@ IPixelCaptureOutputFrame* FPixelCaptureCapturerNV12ToRHI::CreateOutputBuffer(int
 		TextureDesc.AddFlags(ETextureCreateFlags::Shared);
 	}
 
-	return new FPixelCaptureOutputFrameRHI(RHICreateTexture(TextureDesc), { InputWidth, InputHeight, EPixelFormat::PF_NV12 });
+	return new FPixelCaptureOutputFrameRHI(RHICmdList.CreateTexture(TextureDesc), { InputWidth, InputHeight, EPixelFormat::PF_NV12 });
 }
 
-void FPixelCaptureCapturerNV12ToRHI::BeginProcess(const IPixelCaptureInputFrame& InputFrame, TSharedPtr<IPixelCaptureOutputFrame> OutputBuffer)
+void FPixelCaptureCapturerNV12ToRHI::BeginProcess(FRHICommandList& RHICmdList, const IPixelCaptureInputFrame& InputFrame, TSharedPtr<IPixelCaptureOutputFrame> OutputBuffer)
 {
 	SetIsBusy(true);
 
diff --git a/Engine/Plugins/Media/PixelCapture/Source/PixelCapture/Private/PixelCaptureCapturerRHI.cpp b/Engine/Plugins/Media/PixelCapture/Source/PixelCapture/Private/PixelCaptureCapturerRHI.cpp
--- a/Engine/Plugins/Media/PixelCapture/Source/PixelCapture/Private/PixelCaptureCapturerRHI.cpp
+++ b/Engine/Plugins/Media/PixelCapture/Source/PixelCapture/Private/PixelCaptureCapturerRHI.cpp
@@ -19,7 +19,7 @@ FPixelCaptureCapturerRHI::FPixelCaptureCapturerRHI(FPixelCaptureCapturerConfig&
 	RHIType = RHIGetInterfaceType();
 }
 
-IPixelCaptureOutputFrame* FPixelCaptureCapturerRHI::CreateOutputBuffer(int32 InputWidth, int32 InputHeight)
+IPixelCaptureOutputFrame* FPixelCaptureCapturerRHI::CreateOutputBuffer(FRHICommandListBase& RHICmdList, int32 InputWidth, int32 InputHeight)
 {
 	FRHITextureCreateDesc TextureDesc = FRHITextureCreateDesc::Create2D(TEXT("FPixelCaptureCapturerRHI Texture"), Config.OutputResolution.X, Config.OutputResolution.Y, EPixelFormat::PF_B8G8R8A8);
 
@@ -52,17 +52,16 @@ IPixelCaptureOutputFrame* FPixelCaptureCapturerRHI::CreateOutputBuffer(int32 Inp
 		TextureDesc.AddFlags(ETextureCreateFlags::Shared);
 	}
 
-	return new FPixelCaptureOutputFrameRHI(RHICreateTexture(TextureDesc));
+	return new FPixelCaptureOutputFrameRHI(RHICmdList.CreateTexture(TextureDesc));
 }
 
-void FPixelCaptureCapturerRHI::BeginProcess(const IPixelCaptureInputFrame& InputFrame, TSharedPtr<IPixelCaptureOutputFrame> OutputBuffer)
+void FPixelCaptureCapturerRHI::BeginProcess(FRHICommandList& RHICmdList, const IPixelCaptureInputFrame& InputFrame, TSharedPtr<IPixelCaptureOutputFrame> OutputBuffer)
 {
 	checkf(InputFrame.GetType() == StaticCast<int32>(PixelCaptureBufferFormat::FORMAT_RHI), TEXT("Incorrect source frame coming into frame capture process."));
 
 	UE::PixelCapture::MarkCPUWorkStart(OutputBuffer);
 
-	FRHICommandListImmediate& RHICmdList = FRHICommandListExecutor::GetImmediateCommandList();
-	RHICmdList.EnqueueLambda([this, OutputBuffer](FRHICommandListImmediate&) { UE::PixelCapture::MarkGPUWorkStart(OutputBuffer); });
+	RHICmdList.EnqueueLambda([this, OutputBuffer](FRHICommandList&) { UE::PixelCapture::MarkGPUWorkStart(OutputBuffer); });
 
 	const FPixelCaptureInputFrameRHI&		RHISourceFrame = StaticCast<const FPixelCaptureInputFrameRHI&>(InputFrame);
 	TSharedPtr<FPixelCaptureOutputFrameRHI> OutputH264Buffer = StaticCastSharedPtr<FPixelCaptureOutputFrameRHI>(OutputBuffer);
diff --git a/Engine/Plugins/Media/PixelCapture/Source/PixelCapture/Private/PixelCaptureCapturerRHINoCopy.cpp b/Engine/Plugins/Media/PixelCapture/Source/PixelCapture/Private/PixelCaptureCapturerRHINoCopy.cpp
--- a/Engine/Plugins/Media/PixelCapture/Source/PixelCapture/Private/PixelCaptureCapturerRHINoCopy.cpp
+++ b/Engine/Plugins/Media/PixelCapture/Source/PixelCapture/Private/PixelCaptureCapturerRHINoCopy.cpp
@@ -18,12 +18,12 @@ FPixelCaptureCapturerRHINoCopy::FPixelCaptureCapturerRHINoCopy(FPixelCaptureCapt
 {
 }
 
-IPixelCaptureOutputFrame* FPixelCaptureCapturerRHINoCopy::CreateOutputBuffer(int32 InputWidth, int32 InputHeight)
+IPixelCaptureOutputFrame* FPixelCaptureCapturerRHINoCopy::CreateOutputBuffer(FRHICommandListBase& RHICmdList, int32 InputWidth, int32 InputHeight)
 {
 	return new FPixelCaptureOutputFrameRHI(nullptr);
 }
 
-void FPixelCaptureCapturerRHINoCopy::BeginProcess(const IPixelCaptureInputFrame& InputFrame, TSharedPtr<IPixelCaptureOutputFrame> OutputBuffer)
+void FPixelCaptureCapturerRHINoCopy::BeginProcess(FRHICommandList& RHICmdList, const IPixelCaptureInputFrame& InputFrame, TSharedPtr<IPixelCaptureOutputFrame> OutputBuffer)
 {
 	checkf(InputFrame.GetType() == StaticCast<int32>(PixelCaptureBufferFormat::FORMAT_RHI), TEXT("Incorrect source frame coming into frame capture process."));
 
diff --git a/Engine/Plugins/Media/PixelCapture/Source/PixelCapture/Private/PixelCaptureCapturerRHIRDG.cpp b/Engine/Plugins/Media/PixelCapture/Source/PixelCapture/Private/PixelCaptureCapturerRHIRDG.cpp
--- a/Engine/Plugins/Media/PixelCapture/Source/PixelCapture/Private/PixelCaptureCapturerRHIRDG.cpp
+++ b/Engine/Plugins/Media/PixelCapture/Source/PixelCapture/Private/PixelCaptureCapturerRHIRDG.cpp
@@ -18,7 +18,7 @@ FPixelCaptureCapturerRHIRDG::FPixelCaptureCapturerRHIRDG(FPixelCaptureCapturerCo
 {
 }
 
-IPixelCaptureOutputFrame* FPixelCaptureCapturerRHIRDG::CreateOutputBuffer(int32 InputWidth, int32 InputHeight)
+IPixelCaptureOutputFrame* FPixelCaptureCapturerRHIRDG::CreateOutputBuffer(FRHICommandListBase& RHICmdList, int32 InputWidth, int32 InputHeight)
 {
 	FRHITextureCreateDesc TextureDesc = FRHITextureCreateDesc::Create2D(TEXT("FPixelCaptureCapturerRHIRDG Texture"), Config.OutputResolution.X, Config.OutputResolution.Y, EPixelFormat::PF_B8G8R8A8);
 
@@ -51,10 +51,10 @@ IPixelCaptureOutputFrame* FPixelCaptureCapturerRHIRDG::CreateOutputBuffer(int32
 		TextureDesc.AddFlags(ETextureCreateFlags::Shared);
 	}
 
-	return new FPixelCaptureOutputFrameRHI(RHICreateTexture(TextureDesc));
+	return new FPixelCaptureOutputFrameRHI(RHICmdList.CreateTexture(TextureDesc));
 }
 
-void FPixelCaptureCapturerRHIRDG::BeginProcess(const IPixelCaptureInputFrame& InputFrame, TSharedPtr<IPixelCaptureOutputFrame> OutputBuffer)
+void FPixelCaptureCapturerRHIRDG::BeginProcess(FRHICommandList& RHICmdList, const IPixelCaptureInputFrame& InputFrame, TSharedPtr<IPixelCaptureOutputFrame> OutputBuffer)
 {
 	SetIsBusy(true);
 
@@ -63,12 +63,10 @@ void FPixelCaptureCapturerRHIRDG::BeginProcess(const IPixelCaptureInputFrame& In
 	UE::PixelCapture::MarkCPUWorkStart(OutputBuffer);
 	UE::PixelCapture::MarkCPUWorkEnd(OutputBuffer);
 
-	FRHICommandListImmediate& RHICmdList = FRHICommandListExecutor::GetImmediateCommandList();
-
 	UE::PixelCapture::MarkGPUWorkStart(OutputBuffer);
 	const FPixelCaptureInputFrameRHI&		RHISourceFrame = StaticCast<const FPixelCaptureInputFrameRHI&>(InputFrame);
 	TSharedPtr<FPixelCaptureOutputFrameRHI> OutputH264Buffer = StaticCastSharedPtr<FPixelCaptureOutputFrameRHI>(OutputBuffer);
-	CopyTextureRDG(RHICmdList, RHISourceFrame.FrameTexture, OutputH264Buffer->GetFrameTexture());
+	CopyTextureRDG(RHICmdList.GetAsImmediate(), RHISourceFrame.FrameTexture, OutputH264Buffer->GetFrameTexture());
 
 	UE::PixelCapture::MarkGPUWorkEnd(OutputBuffer);
 	OnRHIStageComplete(OutputBuffer);
diff --git a/Engine/Plugins/Media/PixelCapture/Source/PixelCapture/Private/PixelCaptureCapturerRHIToI420CPU.cpp b/Engine/Plugins/Media/PixelCapture/Source/PixelCapture/Private/PixelCaptureCapturerRHIToI420CPU.cpp
--- a/Engine/Plugins/Media/PixelCapture/Source/PixelCapture/Private/PixelCaptureCapturerRHIToI420CPU.cpp
+++ b/Engine/Plugins/Media/PixelCapture/Source/PixelCapture/Private/PixelCaptureCapturerRHIToI420CPU.cpp
@@ -25,14 +25,14 @@ FPixelCaptureCapturerRHIToI420CPU::~FPixelCaptureCapturerRHIToI420CPU()
 {
 }
 
-IPixelCaptureOutputFrame* FPixelCaptureCapturerRHIToI420CPU::CreateOutputBuffer(int32 InputWidth, int32 InputHeight)
+IPixelCaptureOutputFrame* FPixelCaptureCapturerRHIToI420CPU::CreateOutputBuffer(FRHICommandListBase& RHICmdList, int32 InputWidth, int32 InputHeight)
 {
 	return new FPixelCaptureOutputFrameI420(MakeShared<FPixelCaptureBufferI420>(Config.OutputResolution.X, Config.OutputResolution.Y));
 }
 
 #if 1
 // Old code path that uses a FRHIGPUTextureReadback. This seems to induce extra latency but does not introduce screen tearing
-void FPixelCaptureCapturerRHIToI420CPU::Initialize(int32 InputWidth, int32 InputHeight)
+void FPixelCaptureCapturerRHIToI420CPU::Initialize(FRHICommandListBase& RHICmdList, int32 InputWidth, int32 InputHeight)
 {
 	FRHITextureCreateDesc TextureDesc =
 		FRHITextureCreateDesc::Create2D(TEXT("FPixelCaptureCapturerRHIToI420CPU StagingTexture"), Config.OutputResolution.X, Config.OutputResolution.Y, EPixelFormat::PF_B8G8R8A8)
@@ -50,14 +50,14 @@ void FPixelCaptureCapturerRHIToI420CPU::Initialize(int32 InputWidth, int32 Input
 		TextureDesc.AddFlags(ETextureCreateFlags::Shared);
 	}
 
-	StagingTexture = RHICreateTexture(TextureDesc);
+	StagingTexture = RHICmdList.CreateTexture(TextureDesc);
 
 	TextureReader = MakeShared<FRHIGPUTextureReadback>(TEXT("FPixelCaptureCapturerRHIToI420CPUReadback"));
 
-	FPixelCaptureCapturer::Initialize(InputWidth, InputHeight);
+	FPixelCaptureCapturer::Initialize(RHICmdList, InputWidth, InputHeight);
 }
 
-void FPixelCaptureCapturerRHIToI420CPU::BeginProcess(const IPixelCaptureInputFrame& InputFrame, TSharedPtr<IPixelCaptureOutputFrame> OutputBuffer)
+void FPixelCaptureCapturerRHIToI420CPU::BeginProcess(FRHICommandList& RHICmdList, const IPixelCaptureInputFrame& InputFrame, TSharedPtr<IPixelCaptureOutputFrame> OutputBuffer)
 {
 	SetIsBusy(true);
 
@@ -68,9 +68,7 @@ void FPixelCaptureCapturerRHIToI420CPU::BeginProcess(const IPixelCaptureInputFra
 	const FPixelCaptureInputFrameRHI& RHISourceFrame = StaticCast<const FPixelCaptureInputFrameRHI&>(InputFrame);
 	FTextureRHIRef					  SourceTexture = RHISourceFrame.FrameTexture;
 
-	FRHICommandListImmediate& RHICmdList = FRHICommandListExecutor::GetImmediateCommandList();
-
-	RHICmdList.EnqueueLambda([this, OutputBuffer](FRHICommandListImmediate&) { UE::PixelCapture::MarkGPUWorkStart(OutputBuffer); });
+	RHICmdList.EnqueueLambda([this, OutputBuffer](FRHICommandList&) { UE::PixelCapture::MarkGPUWorkStart(OutputBuffer); });
 
 	RHICmdList.Transition(FRHITransitionInfo(SourceTexture, ERHIAccess::Unknown, ERHIAccess::CopySrc));
 	RHICmdList.Transition(FRHITransitionInfo(StagingTexture, ERHIAccess::CopySrc, ERHIAccess::CopyDest));
@@ -84,7 +82,7 @@ void FPixelCaptureCapturerRHIToI420CPU::BeginProcess(const IPixelCaptureInputFra
 	// by adding this shared ref to the rhi lambda we can ensure that 'this' will not be destroyed
 	// until after the rhi thread is done with it, so all the commands will still have valid references.
 	TSharedRef<FPixelCaptureCapturerRHIToI420CPU> ThisRHIRef = StaticCastSharedRef<FPixelCaptureCapturerRHIToI420CPU>(AsShared());
-	RHICmdList.EnqueueLambda([ThisRHIRef, OutputBuffer](FRHICommandListImmediate&) {
+	RHICmdList.EnqueueLambda([ThisRHIRef, OutputBuffer](FRHICommandList&) {
 		ThisRHIRef->CheckComplete(OutputBuffer);
 	});
 }
@@ -133,7 +131,7 @@ void FPixelCaptureCapturerRHIToI420CPU::OnRHIStageComplete(TSharedPtr<IPixelCapt
 // TODO (Eden.Harris) This has less latency but introduces screen tearing so is disabled. RTCP-7778 to fix the screen tearing without effecting latency.
 #else
 
-void FPixelCaptureCapturerRHIToI420CPU::Initialize(int32 InputWidth, int32 InputHeight)
+void FPixelCaptureCapturerRHIToI420CPU::Initialize(FRHICommandListBase& RHICmdList, int32 InputWidth, int32 InputHeight)
 {
 	FRHITextureCreateDesc StagingTextureDesc =
 		FRHITextureCreateDesc::Create2D(TEXT("FPixelCaptureCapturerRHIToI420CPU StagingTexture"), Config.OutputResolution.X, Config.OutputResolution.Y, EPixelFormat::PF_B8G8R8A8)
@@ -151,7 +149,7 @@ void FPixelCaptureCapturerRHIToI420CPU::Initialize(int32 InputWidth, int32 Input
 		StagingTextureDesc.AddFlags(ETextureCreateFlags::Shared);
 	}
 
-	StagingTexture = RHICreateTexture(StagingTextureDesc);
+	StagingTexture = RHICmdList.CreateTexture(StagingTextureDesc);
 
 	FRHITextureCreateDesc ReadbackTextureDesc = FRHITextureCreateDesc::Create2D(TEXT("FPixelCaptureCapturerRHIToI420CPU ReadbackTexture"), Config.OutputResolution.X, Config.OutputResolution.Y, EPixelFormat::PF_B8G8R8A8)
 													.SetClearValue(FClearValueBinding::None)
@@ -159,12 +157,12 @@ void FPixelCaptureCapturerRHIToI420CPU::Initialize(int32 InputWidth, int32 Input
 													.SetInitialState(ERHIAccess::CopyDest)
 													.DetermineInititialState();
 
-	ReadbackTexture = RHICreateTexture(ReadbackTextureDesc);
+	ReadbackTexture = RHICmdList.CreateTexture(ReadbackTextureDesc);
 
 	FPixelCaptureCapturer::Initialize(InputWidth, InputHeight);
 }
 
-void FPixelCaptureCapturerRHIToI420CPU::BeginProcess(const IPixelCaptureInputFrame& InputFrame, TSharedPtr<IPixelCaptureOutputFrame> OutputBuffer)
+void FPixelCaptureCapturerRHIToI420CPU::BeginProcess(FRHICommandList& RHICmdList, const IPixelCaptureInputFrame& InputFrame, TSharedPtr<IPixelCaptureOutputFrame> OutputBuffer)
 {
 	SetIsBusy(true);
 
@@ -173,7 +171,6 @@ void FPixelCaptureCapturerRHIToI420CPU::BeginProcess(const IPixelCaptureInputFra
 	const FPixelCaptureInputFrameRHI& RHISourceFrame = StaticCast<const FPixelCaptureInputFrameRHI&>(InputFrame);
 	FTextureRHIRef					  SourceTexture = RHISourceFrame.FrameTexture;
 
-	FRHICommandListImmediate& RHICmdList = FRHICommandListExecutor::GetImmediateCommandList();
 	RHICmdList.Transition(FRHITransitionInfo(SourceTexture, ERHIAccess::Unknown, ERHIAccess::CopySrc));
 	RHICmdList.Transition(FRHITransitionInfo(StagingTexture, ERHIAccess::CopySrc, ERHIAccess::CopyDest));
 	CopyTexture(RHICmdList, SourceTexture, StagingTexture, nullptr);
@@ -184,7 +181,7 @@ void FPixelCaptureCapturerRHIToI420CPU::BeginProcess(const IPixelCaptureInputFra
 	void* ReadbackPointer = nullptr;
 	int32 ReadbackWidth, ReadbackHeight;
 	// Nullptr fence forces a flush
-	RHICmdList.MapStagingSurface(ReadbackTexture, nullptr, ReadbackPointer, ReadbackWidth, ReadbackHeight);
+	RHICmdList.GetAsImmediate().MapStagingSurface(ReadbackTexture, nullptr, ReadbackPointer, ReadbackWidth, ReadbackHeight);
 
 	TSharedPtr<FPixelCaptureOutputFrameI420> OutputI420Buffer = StaticCastSharedPtr<FPixelCaptureOutputFrameI420>(OutputBuffer);
 	TSharedPtr<FPixelCaptureBufferI420>		 I420Buffer = OutputI420Buffer->GetI420Buffer();
@@ -200,7 +197,7 @@ void FPixelCaptureCapturerRHIToI420CPU::BeginProcess(const IPixelCaptureInputFra
 		I420Buffer->GetWidth(),
 		I420Buffer->GetHeight());
 
-	RHICmdList.UnmapStagingSurface(ReadbackTexture);
+	RHICmdList.GetAsImmediate().UnmapStagingSurface(ReadbackTexture);
 
 	EndProcess(OutputBuffer);
 	SetIsBusy(false);
diff --git a/Engine/Plugins/Media/PixelCapture/Source/PixelCapture/Private/PixelCaptureCapturerRHIToI420Compute.cpp b/Engine/Plugins/Media/PixelCapture/Source/PixelCapture/Private/PixelCaptureCapturerRHIToI420Compute.cpp
--- a/Engine/Plugins/Media/PixelCapture/Source/PixelCapture/Private/PixelCaptureCapturerRHIToI420Compute.cpp
+++ b/Engine/Plugins/Media/PixelCapture/Source/PixelCapture/Private/PixelCaptureCapturerRHIToI420Compute.cpp
@@ -36,10 +36,8 @@ FPixelCaptureCapturerRHIToI420Compute ::~FPixelCaptureCapturerRHIToI420Compute()
 	CleanUp();
 }
 
-void FPixelCaptureCapturerRHIToI420Compute::Initialize(int32 InputWidth, int32 InputHeight)
+void FPixelCaptureCapturerRHIToI420Compute::Initialize(FRHICommandListBase& RHICmdList, int32 InputWidth, int32 InputHeight)
 {
-	FRHICommandListImmediate& RHICmdList = FRHICommandListImmediate::Get();
-
 	PlaneYDimensions = { Config.OutputResolution.X, Config.OutputResolution.Y };
 	PlaneUVDimensions = { (Config.OutputResolution.X + 1) / 2, (Config.OutputResolution.Y + 1) / 2 }; // UV is halved and rounded up
 
@@ -57,9 +55,9 @@ void FPixelCaptureCapturerRHIToI420Compute::Initialize(int32 InputWidth, int32 I
 			.SetInitialState(ERHIAccess::UAVCompute)
 			.DetermineInititialState();
 
-	TextureY = RHICreateTexture(TextureDescY);
-	TextureU = RHICreateTexture(TextureDescUV);
-	TextureV = RHICreateTexture(TextureDescUV);
+	TextureY = RHICmdList.CreateTexture(TextureDescY);
+	TextureU = RHICmdList.CreateTexture(TextureDescUV);
+	TextureV = RHICmdList.CreateTexture(TextureDescUV);
 
 	FRHITextureCreateDesc StagingDescY =
 		FRHITextureCreateDesc::Create2D(TEXT("YUV Output CPU Texture"), PlaneYDimensions.X, PlaneYDimensions.Y, EPixelFormat::PF_R8)
@@ -75,31 +73,31 @@ void FPixelCaptureCapturerRHIToI420Compute::Initialize(int32 InputWidth, int32 I
 			.SetInitialState(ERHIAccess::Unknown)
 			.DetermineInititialState();
 
-	StagingTextureY = RHICreateTexture(StagingDescY);
-	StagingTextureU = RHICreateTexture(StagingDescUV);
-	StagingTextureV = RHICreateTexture(StagingDescUV);
+	StagingTextureY = RHICmdList.CreateTexture(StagingDescY);
+	StagingTextureU = RHICmdList.CreateTexture(StagingDescUV);
+	StagingTextureV = RHICmdList.CreateTexture(StagingDescUV);
 
 	TextureYUAV = RHICmdList.CreateUnorderedAccessView(TextureY, FRHIViewDesc::CreateTextureUAV().SetDimensionFromTexture(TextureY));
 	TextureUUAV = RHICmdList.CreateUnorderedAccessView(TextureU, FRHIViewDesc::CreateTextureUAV().SetDimensionFromTexture(TextureU));
 	TextureVUAV = RHICmdList.CreateUnorderedAccessView(TextureV, FRHIViewDesc::CreateTextureUAV().SetDimensionFromTexture(TextureV));
 
 	int32 OutWidth, OutHeight;
-	RHICmdList.MapStagingSurface(StagingTextureY, nullptr, MappedY, OutWidth, OutHeight);
+	RHICmdList.GetAsImmediate().MapStagingSurface(StagingTextureY, nullptr, MappedY, OutWidth, OutHeight);
 	YStride = OutWidth;
-	RHICmdList.MapStagingSurface(StagingTextureU, nullptr, MappedU, OutWidth, OutHeight);
+	RHICmdList.GetAsImmediate().MapStagingSurface(StagingTextureU, nullptr, MappedU, OutWidth, OutHeight);
 	UStride = OutWidth;
-	RHICmdList.MapStagingSurface(StagingTextureV, nullptr, MappedV, OutWidth, OutHeight);
+	RHICmdList.GetAsImmediate().MapStagingSurface(StagingTextureV, nullptr, MappedV, OutWidth, OutHeight);
 	VStride = OutWidth;
 
-	FPixelCaptureCapturer::Initialize(InputWidth, InputHeight);
+	FPixelCaptureCapturer::Initialize(RHICmdList, InputWidth, InputHeight);
 }
 
-IPixelCaptureOutputFrame* FPixelCaptureCapturerRHIToI420Compute::CreateOutputBuffer(int32 InputWidth, int32 InputHeight)
+IPixelCaptureOutputFrame* FPixelCaptureCapturerRHIToI420Compute::CreateOutputBuffer(FRHICommandListBase& RHICmdList, int32 InputWidth, int32 InputHeight)
 {
 	return new FPixelCaptureOutputFrameI420(MakeShared<FPixelCaptureBufferI420>(Config.OutputResolution.X, Config.OutputResolution.Y));
 }
 
-void FPixelCaptureCapturerRHIToI420Compute::BeginProcess(const IPixelCaptureInputFrame& InputFrame, TSharedPtr<IPixelCaptureOutputFrame> OutputBuffer)
+void FPixelCaptureCapturerRHIToI420Compute::BeginProcess(FRHICommandList& RHICmdList, const IPixelCaptureInputFrame& InputFrame, TSharedPtr<IPixelCaptureOutputFrame> OutputBuffer)
 {
 	SetIsBusy(true);
 
@@ -110,9 +108,7 @@ void FPixelCaptureCapturerRHIToI420Compute::BeginProcess(const IPixelCaptureInpu
 	const FPixelCaptureInputFrameRHI& RHISourceFrame = StaticCast<const FPixelCaptureInputFrameRHI&>(InputFrame);
 	FTextureRHIRef					  SourceTexture = RHISourceFrame.FrameTexture;
 
-	FRHICommandListImmediate& RHICmdList = FRHICommandListExecutor::GetImmediateCommandList();
-
-	RHICmdList.EnqueueLambda([this, OutputBuffer](FRHICommandListImmediate&) { UE::PixelCapture::MarkGPUWorkStart(OutputBuffer); });
+	RHICmdList.EnqueueLambda([this, OutputBuffer](FRHICommandList&) { UE::PixelCapture::MarkGPUWorkStart(OutputBuffer); });
 
 	FRGBToYUVShaderParameters ShaderParameters;
 	ShaderParameters.SourceTexture = SourceTexture;
@@ -132,7 +128,7 @@ void FPixelCaptureCapturerRHIToI420Compute::BeginProcess(const IPixelCaptureInpu
 	// by adding this shared ref to the rhi lambda we can ensure that 'this' will not be destroyed
 	// until after the rhi thread is done with it, so all the commands will still have valid references.
 	TSharedRef<FPixelCaptureCapturerRHIToI420Compute> ThisRHIRef = StaticCastSharedRef<FPixelCaptureCapturerRHIToI420Compute>(AsShared());
-	RHICmdList.EnqueueLambda([ThisRHIRef, OutputBuffer](FRHICommandListImmediate&) { ThisRHIRef->OnRHIStageComplete(OutputBuffer); });
+	RHICmdList.EnqueueLambda([ThisRHIRef, OutputBuffer](FRHICommandList&) { ThisRHIRef->OnRHIStageComplete(OutputBuffer); });
 }
 
 void FPixelCaptureCapturerRHIToI420Compute::OnRHIStageComplete(TSharedPtr<IPixelCaptureOutputFrame> OutputBuffer)
diff --git a/Engine/Plugins/Media/PixelCapture/Source/PixelCapture/Private/Tests/PixelCaptureTests.cpp b/Engine/Plugins/Media/PixelCapture/Source/PixelCapture/Private/Tests/PixelCaptureTests.cpp
--- a/Engine/Plugins/Media/PixelCapture/Source/PixelCapture/Private/Tests/PixelCaptureTests.cpp
+++ b/Engine/Plugins/Media/PixelCapture/Source/PixelCapture/Private/Tests/PixelCaptureTests.cpp
@@ -6,6 +6,7 @@
 #include "PixelCaptureCapturer.h"
 #include "PixelCaptureCapturerLayered.h"
 #include "PixelCaptureCapturerMultiFormat.h"
+#include "RHICommandList.h"
 
 #if WITH_DEV_AUTOMATION_TESTS
 
@@ -62,12 +63,12 @@ namespace
 
 		virtual FString GetCapturerName() const override { return "Mock Capturer"; }
 
-		virtual IPixelCaptureOutputFrame* CreateOutputBuffer(int32 InputWidth, int32 InputHeight) override
+		virtual IPixelCaptureOutputFrame* CreateOutputBuffer(FRHICommandListBase& RHICmdList, int32 InputWidth, int32 InputHeight) override
 		{
 			return new FMockOutput();
 		}
 
-		virtual void BeginProcess(const IPixelCaptureInputFrame& InputFrame, TSharedPtr<IPixelCaptureOutputFrame> OutputBuffer) override
+		virtual void BeginProcess(FRHICommandList& RHICmdList, const IPixelCaptureInputFrame& InputFrame, TSharedPtr<IPixelCaptureOutputFrame> OutputBuffer) override
 		{
 			SetIsBusy(true);
 
@@ -142,6 +143,8 @@ namespace UE::PixelCapture
 		const int32 InputHeight = 768;
 		const int32 InputMagicData = 0xDEADBEEF;
 
+		FRHICommandListImmediate& RHICmdList = FRHICommandListExecutor::GetImmediateCommandList();
+
 		// basic capture behaviour
 		{
 			FMockCapturer Capturer(StaticCast<int32>(EMockFormats::Output1), 1.0f, false);
@@ -149,7 +152,7 @@ namespace UE::PixelCapture
 			TestTrue("Capturer starts not busy", Capturer.IsBusy() == false);
 			TestTrue("Capturer starts without output", Capturer.HasOutput() == false);
 			TestTrue("Capturer returns null if no input present", Capturer.ReadOutput() == nullptr);
-			Capturer.Capture(FMockInput(InputWidth, InputHeight, InputMagicData));
+			Capturer.Capture(RHICmdList, FMockInput(InputWidth, InputHeight, InputMagicData));
 			TestTrue("Capturer initialized after frame input", Capturer.IsInitialized() == true);
 			TestTrue("Capturer busy after frame input", Capturer.IsBusy() == true);
 			TestTrue("Capturer has no output before completing", Capturer.HasOutput() == false);
@@ -170,7 +173,7 @@ namespace UE::PixelCapture
 
 			TSharedPtr<FPixelCaptureCapturerLayered> LayeredCapturer = FPixelCaptureCapturerLayered::Create(&CaptureSource, StaticCast<int32>(EMockFormats::Output1), Resolutions);
 			LayeredCapturer->OnComplete.AddLambda([&]() { CaptureSource.FrameCompleteCount++; });
-			LayeredCapturer->Capture(FMockInput(InputWidth, InputHeight, InputMagicData));
+			LayeredCapturer->Capture(RHICmdList, FMockInput(InputWidth, InputHeight, InputMagicData));
 			TestTrue("Frame complete callback called once.", CaptureSource.FrameCompleteCount == 3);
 			for (TMap<FMockCapturer*, int>::ElementType& CaptureComplete : CaptureSource.LayerCompleteCounts)
 			{
@@ -191,7 +194,7 @@ namespace UE::PixelCapture
 
 			TSharedPtr<FPixelCaptureCapturerLayered> LayeredCapturer = FPixelCaptureCapturerLayered::Create(&CaptureSource, StaticCast<int32>(EMockFormats::Output1));
 			LayeredCapturer->OnComplete.AddLambda([&]() { CaptureSource.FrameCompleteCount++; });
-			LayeredCapturer->Capture(FMockInput(InputWidth, InputHeight, InputMagicData));
+			LayeredCapturer->Capture(RHICmdList, FMockInput(InputWidth, InputHeight, InputMagicData));
 			TestTrue("Frame complete callback not called as there's no known output resolutions.", CaptureSource.FrameCompleteCount == 0);
 			for (int i = 0; i < Resolutions.Num(); ++i)
 			{
@@ -200,7 +203,7 @@ namespace UE::PixelCapture
 				TestTrue(FString::Printf(TEXT("Layer %d has no output"), i), !Layer.IsValid());
 			}
 
-			LayeredCapturer->Capture(FMockInput(InputWidth, InputHeight, InputMagicData));
+			LayeredCapturer->Capture(RHICmdList, FMockInput(InputWidth, InputHeight, InputMagicData));
 			TestTrue("Frame complete callback called once per output resolution.", CaptureSource.FrameCompleteCount == 3);
 
 			for (TMap<FMockCapturer*, int>::ElementType& CaptureComplete : CaptureSource.LayerCompleteCounts)
@@ -231,7 +234,7 @@ namespace UE::PixelCapture
 
 			// Test output after implicit format add from WaitForFormat
 			{
-				FormatCapturer->Capture(FMockInput(InputWidth, InputHeight, InputMagicData));
+				FormatCapturer->Capture(RHICmdList, FMockInput(InputWidth, InputHeight, InputMagicData));
 				TestTrue("MultiFormat Capturer reports correct layer count after input.", FormatCapturer->GetNumLayers() == Resolutions.Num());
 				for (int i = 0; i < Resolutions.Num(); ++i)
 				{
@@ -248,7 +251,7 @@ namespace UE::PixelCapture
 		{
 			const int32 TestFormat = StaticCast<int>(EMockFormats::Output2);
 			FormatCapturer->AddOutputFormat(TestFormat);
-			FormatCapturer->Capture(FMockInput(InputWidth, InputHeight, InputMagicData));
+			FormatCapturer->Capture(RHICmdList, FMockInput(InputWidth, InputHeight, InputMagicData));
 			for (int i = 0; i < Resolutions.Num(); ++i)
 			{
 				TSharedPtr<FMockOutput> Layer = StaticCastSharedPtr<FMockOutput>(FormatCapturer->RequestFormat(TestFormat, Resolutions[i]));
@@ -260,7 +263,7 @@ namespace UE::PixelCapture
 		// make sure all formats get new frames
 		{
 			const int32 InputMagicData2 = 0xBAADF00D;
-			FormatCapturer->Capture(FMockInput(InputWidth, InputHeight, InputMagicData2));
+			FormatCapturer->Capture(RHICmdList, FMockInput(InputWidth, InputHeight, InputMagicData2));
 
 			auto CheckFormats = { StaticCast<int>(EMockFormats::Output1), StaticCast<int>(EMockFormats::Output2) };
 			for (auto& TestFormat : CheckFormats)
@@ -285,7 +288,7 @@ namespace UE::PixelCapture
 
 			FormatCapturer->AddOutputFormat(StaticCast<int>(EMockFormats::Output1));
 			FormatCapturer->AddOutputFormat(StaticCast<int>(EMockFormats::Output2));
-			FormatCapturer->Capture(FMockInput(InputWidth, InputHeight, InputMagicData));
+			FormatCapturer->Capture(RHICmdList, FMockInput(InputWidth, InputHeight, InputMagicData));
 
 			TestTrue("MultiFormat Frame complete callback called once.", CaptureSource.FrameCompleteCount == 6);
 			for (TMap<FMockCapturer*, int>::ElementType& CaptureComplete : CaptureSource.LayerCompleteCounts)
diff --git a/Engine/Plugins/Media/PixelCapture/Source/PixelCapture/Public/PixelCaptureCapturer.h b/Engine/Plugins/Media/PixelCapture/Source/PixelCapture/Public/PixelCaptureCapturer.h
--- a/Engine/Plugins/Media/PixelCapture/Source/PixelCapture/Public/PixelCaptureCapturer.h
+++ b/Engine/Plugins/Media/PixelCapture/Source/PixelCapture/Public/PixelCaptureCapturer.h
@@ -6,6 +6,9 @@
 #include "IPixelCaptureOutputFrame.h"
 #include "HAL/ThreadSafeCounter64.h"
 
+class FRHICommandList;
+class FRHICommandListBase;
+
 #define UE_API PIXELCAPTURE_API
 
 namespace UE::PixelCapture
@@ -47,7 +50,7 @@ public:
 	 * Called when an input frame needs capturing.
 	 * @param InputFrame The input frame to be captured.
 	 */
-	UE_API void Capture(const IPixelCaptureInputFrame& InputFrame);
+	UE_API void Capture(FRHICommandList& RHICmdList, const IPixelCaptureInputFrame& InputFrame);
 
 	/**
 	 * Returns true if Initialize() has been called.
@@ -99,7 +102,7 @@ protected:
 	 * @param InputWidth The pixel count of the input frame width
 	 * @param InputHeight The pixel count of the input frame height
 	 */
-	UE_API virtual void Initialize(int32 InputWidth, int32 InputHeight);
+	UE_API virtual void Initialize(FRHICommandListBase& RHICmdList, int32 InputWidth, int32 InputHeight);
 
 	/**
 	 * Implement this to create a buffer for the output.
@@ -107,7 +110,7 @@ protected:
 	 * @param InputHeight The pixel height of the input frame.
 	 * @return An empty output structure that the process can store the output of its process on.
 	 */
-	virtual IPixelCaptureOutputFrame* CreateOutputBuffer(int32 InputWidth, int32 InputHeight) = 0;
+	virtual IPixelCaptureOutputFrame* CreateOutputBuffer(FRHICommandListBase& RHICmdList, int32 InputWidth, int32 InputHeight) = 0;
 
 	/**
 	 * Implement this with your specific process to capture the incoming frame.
@@ -116,7 +119,7 @@ protected:
 	 */
 	UE_DEPRECATED(5.6, "BeginProcess(const IPixelCaptureInputFrame& InputFrame, IPixelCaptureOutputFrame* OutputBuffer) has been deprecated. Please use BeginProcess(const IPixelCaptureInputFrame& InputFrame, TSharedPtr<IPixelCaptureOutputFrame> OutputBuffer) to ensure memory safety.")
 	void		 BeginProcess(const IPixelCaptureInputFrame& InputFrame, IPixelCaptureOutputFrame* OutputBuffer) {};
-	virtual void BeginProcess(const IPixelCaptureInputFrame& InputFrame, TSharedPtr<IPixelCaptureOutputFrame> OutputBuffer) = 0;
+	virtual void BeginProcess(FRHICommandList& RHICmdList, const IPixelCaptureInputFrame& InputFrame, TSharedPtr<IPixelCaptureOutputFrame> OutputBuffer) = 0;
 
 	/**
 	 * Metadata info (optional). Marks the start of some CPU work. Multiple work sections are valid.
diff --git a/Engine/Plugins/Media/PixelCapture/Source/PixelCapture/Public/PixelCaptureCapturerI420.h b/Engine/Plugins/Media/PixelCapture/Source/PixelCapture/Public/PixelCaptureCapturerI420.h
--- a/Engine/Plugins/Media/PixelCapture/Source/PixelCapture/Public/PixelCaptureCapturerI420.h
+++ b/Engine/Plugins/Media/PixelCapture/Source/PixelCapture/Public/PixelCaptureCapturerI420.h
@@ -18,10 +18,10 @@ public:
 	virtual ~FPixelCaptureCapturerI420() = default;
 
 protected:
-	virtual FString					  GetCapturerName() const override { return "I420 Copy"; }
-	UE_API virtual void					  Initialize(int32 InputWidth, int32 InputHeight) override;
-	UE_API virtual IPixelCaptureOutputFrame* CreateOutputBuffer(int32 InputWidth, int32 InputHeight) override;
-	UE_API virtual void					  BeginProcess(const IPixelCaptureInputFrame& InputFrame, TSharedPtr<IPixelCaptureOutputFrame> OutputBuffer) override;
+	virtual FString                          GetCapturerName() const override { return "I420 Copy"; }
+	UE_API virtual void                      Initialize(FRHICommandListBase& RHICmdList, int32 InputWidth, int32 InputHeight) override;
+	UE_API virtual IPixelCaptureOutputFrame* CreateOutputBuffer(FRHICommandListBase& RHICmdList, int32 InputWidth, int32 InputHeight) override;
+	UE_API virtual void                      BeginProcess(FRHICommandList& RHICmdList, const IPixelCaptureInputFrame& InputFrame, TSharedPtr<IPixelCaptureOutputFrame> OutputBuffer) override;
 
 private:
 	UE_API FPixelCaptureCapturerI420(FPixelCaptureCapturerConfig& Config);
diff --git a/Engine/Plugins/Media/PixelCapture/Source/PixelCapture/Public/PixelCaptureCapturerI420ToRHI.h b/Engine/Plugins/Media/PixelCapture/Source/PixelCapture/Public/PixelCaptureCapturerI420ToRHI.h
--- a/Engine/Plugins/Media/PixelCapture/Source/PixelCapture/Public/PixelCaptureCapturerI420ToRHI.h
+++ b/Engine/Plugins/Media/PixelCapture/Source/PixelCapture/Public/PixelCaptureCapturerI420ToRHI.h
@@ -19,9 +19,9 @@ public:
 
 protected:
 	virtual FString GetCapturerName() const override { return "I420ToRHI"; }
-	UE_API virtual void Initialize(int32 InputWidth, int32 InputHeight) override;
-	UE_API virtual IPixelCaptureOutputFrame* CreateOutputBuffer(int32 InputWidth, int32 InputHeight) override;
-	UE_API virtual void BeginProcess(const IPixelCaptureInputFrame& InputFrame, TSharedPtr<IPixelCaptureOutputFrame> OutputBuffer) override;
+	UE_API virtual void Initialize(FRHICommandListBase& RHICmdList, int32 InputWidth, int32 InputHeight) override;
+	UE_API virtual IPixelCaptureOutputFrame* CreateOutputBuffer(FRHICommandListBase& RHICmdList, int32 InputWidth, int32 InputHeight) override;
+	UE_API virtual void BeginProcess(FRHICommandList& RHICmdList, const IPixelCaptureInputFrame& InputFrame, TSharedPtr<IPixelCaptureOutputFrame> OutputBuffer) override;
 
 private:
 	uint8_t* ARGBBuffer = nullptr;
diff --git a/Engine/Plugins/Media/PixelCapture/Source/PixelCapture/Public/PixelCaptureCapturerLayered.h b/Engine/Plugins/Media/PixelCapture/Source/PixelCapture/Public/PixelCaptureCapturerLayered.h
--- a/Engine/Plugins/Media/PixelCapture/Source/PixelCapture/Public/PixelCaptureCapturerLayered.h
+++ b/Engine/Plugins/Media/PixelCapture/Source/PixelCapture/Public/PixelCaptureCapturerLayered.h
@@ -35,7 +35,7 @@ public:
 	/**
 	 * Begins the capture process of a given frame.
 	 */
-	void Capture(const IPixelCaptureInputFrame& SourceFrame);
+	void Capture(FRHICommandList& RHICmdList, const IPixelCaptureInputFrame& SourceFrame);
 
 	/**
 	 * Try to read the result of the capture process. May return null if no output
diff --git a/Engine/Plugins/Media/PixelCapture/Source/PixelCapture/Public/PixelCaptureCapturerMultiFormat.h b/Engine/Plugins/Media/PixelCapture/Source/PixelCapture/Public/PixelCaptureCapturerMultiFormat.h
--- a/Engine/Plugins/Media/PixelCapture/Source/PixelCapture/Public/PixelCaptureCapturerMultiFormat.h
+++ b/Engine/Plugins/Media/PixelCapture/Source/PixelCapture/Public/PixelCaptureCapturerMultiFormat.h
@@ -62,7 +62,7 @@ public:
 	/**
 	 * Begins the capture process of a given frame.
 	 */
-	UE_API void Capture(const IPixelCaptureInputFrame& SourceFrame);
+	UE_API void Capture(FRHICommandList& RHICmdList, const IPixelCaptureInputFrame& SourceFrame);
 
 	/**
 	 * Sets up a capture pipeline for the given destination format. No effect if the
diff --git a/Engine/Plugins/Media/PixelCapture/Source/PixelCapture/Public/PixelCaptureCapturerNV12ToRHI.h b/Engine/Plugins/Media/PixelCapture/Source/PixelCapture/Public/PixelCaptureCapturerNV12ToRHI.h
--- a/Engine/Plugins/Media/PixelCapture/Source/PixelCapture/Public/PixelCaptureCapturerNV12ToRHI.h
+++ b/Engine/Plugins/Media/PixelCapture/Source/PixelCapture/Public/PixelCaptureCapturerNV12ToRHI.h
@@ -19,9 +19,9 @@ public:
 
 protected:
 	virtual FString GetCapturerName() const override { return "NV12ToRHI"; }
-	UE_API virtual void Initialize(int32 InputWidth, int32 InputHeight) override;
-	UE_API virtual IPixelCaptureOutputFrame* CreateOutputBuffer(int32 InputWidth, int32 InputHeight) override;
-	UE_API virtual void BeginProcess(const IPixelCaptureInputFrame& InputFrame, TSharedPtr<IPixelCaptureOutputFrame> OutputBuffer) override;
+	UE_API virtual void Initialize(FRHICommandListBase& RHICmdList, int32 InputWidth, int32 InputHeight) override;
+	UE_API virtual IPixelCaptureOutputFrame* CreateOutputBuffer(FRHICommandListBase& RHICmdList, int32 InputWidth, int32 InputHeight) override;
+	UE_API virtual void BeginProcess(FRHICommandList& RHICmdList, const IPixelCaptureInputFrame& InputFrame, TSharedPtr<IPixelCaptureOutputFrame> OutputBuffer) override;
 
 private:
 	uint8_t* R8Buffer = nullptr;
diff --git a/Engine/Plugins/Media/PixelCapture/Source/PixelCapture/Public/PixelCaptureCapturerRHI.h b/Engine/Plugins/Media/PixelCapture/Source/PixelCapture/Public/PixelCaptureCapturerRHI.h
--- a/Engine/Plugins/Media/PixelCapture/Source/PixelCapture/Public/PixelCaptureCapturerRHI.h
+++ b/Engine/Plugins/Media/PixelCapture/Source/PixelCapture/Public/PixelCaptureCapturerRHI.h
@@ -23,9 +23,9 @@ public:
 	virtual ~FPixelCaptureCapturerRHI() = default;
 
 protected:
-	virtual FString					  GetCapturerName() const override { return "RHI Copy"; }
-	UE_API virtual IPixelCaptureOutputFrame* CreateOutputBuffer(int32 InputWidth, int32 InputHeight) override;
-	UE_API virtual void					  BeginProcess(const IPixelCaptureInputFrame& InputFrame, TSharedPtr<IPixelCaptureOutputFrame> OutputBuffer) override;
+	virtual FString                          GetCapturerName() const override { return "RHI Copy"; }
+	UE_API virtual IPixelCaptureOutputFrame* CreateOutputBuffer(FRHICommandListBase& RHICmdList, int32 InputWidth, int32 InputHeight) override;
+	UE_API virtual void                      BeginProcess(FRHICommandList& RHICmdList, const IPixelCaptureInputFrame& InputFrame, TSharedPtr<IPixelCaptureOutputFrame> OutputBuffer) override;
 
 private:
 	ERHIInterfaceType RHIType;
diff --git a/Engine/Plugins/Media/PixelCapture/Source/PixelCapture/Public/PixelCaptureCapturerRHINoCopy.h b/Engine/Plugins/Media/PixelCapture/Source/PixelCapture/Public/PixelCaptureCapturerRHINoCopy.h
--- a/Engine/Plugins/Media/PixelCapture/Source/PixelCapture/Public/PixelCaptureCapturerRHINoCopy.h
+++ b/Engine/Plugins/Media/PixelCapture/Source/PixelCapture/Public/PixelCaptureCapturerRHINoCopy.h
@@ -28,9 +28,9 @@ public:
 	virtual ~FPixelCaptureCapturerRHINoCopy() = default;
 
 protected:
-	virtual FString					  GetCapturerName() const override { return "RHI No Copy"; }
-	UE_API virtual IPixelCaptureOutputFrame* CreateOutputBuffer(int32 InputWidth, int32 InputHeight) override;
-	UE_API virtual void					  BeginProcess(const IPixelCaptureInputFrame& InputFrame, TSharedPtr<IPixelCaptureOutputFrame> OutputBuffer) override;
+	virtual FString                          GetCapturerName() const override { return "RHI No Copy"; }
+	UE_API virtual IPixelCaptureOutputFrame* CreateOutputBuffer(FRHICommandListBase& RHICmdList, int32 InputWidth, int32 InputHeight) override;
+	UE_API virtual void                      BeginProcess(FRHICommandList& RHICmdList, const IPixelCaptureInputFrame& InputFrame, TSharedPtr<IPixelCaptureOutputFrame> OutputBuffer) override;
 
 private:
 	UE_API FPixelCaptureCapturerRHINoCopy(FPixelCaptureCapturerConfig& Config);
diff --git a/Engine/Plugins/Media/PixelCapture/Source/PixelCapture/Public/PixelCaptureCapturerRHIRDG.h b/Engine/Plugins/Media/PixelCapture/Source/PixelCapture/Public/PixelCaptureCapturerRHIRDG.h
--- a/Engine/Plugins/Media/PixelCapture/Source/PixelCapture/Public/PixelCaptureCapturerRHIRDG.h
+++ b/Engine/Plugins/Media/PixelCapture/Source/PixelCapture/Public/PixelCaptureCapturerRHIRDG.h
@@ -23,9 +23,9 @@ public:
 	virtual ~FPixelCaptureCapturerRHIRDG() = default;
 
 protected:
-	virtual FString					  GetCapturerName() const override { return "RHI RDG Copy"; }
-	UE_API virtual IPixelCaptureOutputFrame* CreateOutputBuffer(int32 InputWidth, int32 InputHeight) override;
-	UE_API virtual void					  BeginProcess(const IPixelCaptureInputFrame& InputFrame, TSharedPtr<IPixelCaptureOutputFrame> OutputBuffer) override;
+	virtual FString                          GetCapturerName() const override { return "RHI RDG Copy"; }
+	UE_API virtual IPixelCaptureOutputFrame* CreateOutputBuffer(FRHICommandListBase& RHICmdList, int32 InputWidth, int32 InputHeight) override;
+	UE_API virtual void                      BeginProcess(FRHICommandList& RHICmdList, const IPixelCaptureInputFrame& InputFrame, TSharedPtr<IPixelCaptureOutputFrame> OutputBuffer) override;
 
 private:
 	UE_API FPixelCaptureCapturerRHIRDG(FPixelCaptureCapturerConfig& Config);
diff --git a/Engine/Plugins/Media/PixelCapture/Source/PixelCapture/Public/PixelCaptureCapturerRHIToI420CPU.h b/Engine/Plugins/Media/PixelCapture/Source/PixelCapture/Public/PixelCaptureCapturerRHIToI420CPU.h
--- a/Engine/Plugins/Media/PixelCapture/Source/PixelCapture/Public/PixelCaptureCapturerRHIToI420CPU.h
+++ b/Engine/Plugins/Media/PixelCapture/Source/PixelCapture/Public/PixelCaptureCapturerRHIToI420CPU.h
@@ -25,10 +25,10 @@ public:
 	UE_API virtual ~FPixelCaptureCapturerRHIToI420CPU();
 
 protected:
-	virtual FString					  GetCapturerName() const override { return "RHIToI420CPU"; }
-	UE_API virtual void					  Initialize(int32 InputWidth, int32 InputHeight) override;
-	UE_API virtual IPixelCaptureOutputFrame* CreateOutputBuffer(int32 InputWidth, int32 InputHeight) override;
-	UE_API virtual void					  BeginProcess(const IPixelCaptureInputFrame& InputFrame, TSharedPtr<IPixelCaptureOutputFrame> OutputBuffer) override;
+	virtual FString                          GetCapturerName() const override { return "RHIToI420CPU"; }
+	UE_API virtual void                      Initialize(FRHICommandListBase& RHICmdList, int32 InputWidth, int32 InputHeight) override;
+	UE_API virtual IPixelCaptureOutputFrame* CreateOutputBuffer(FRHICommandListBase& RHICmdList, int32 InputWidth, int32 InputHeight) override;
+	UE_API virtual void                      BeginProcess(FRHICommandList& RHICmdList, const IPixelCaptureInputFrame& InputFrame, TSharedPtr<IPixelCaptureOutputFrame> OutputBuffer) override;
 
 private:
 	FTextureRHIRef StagingTexture;
diff --git a/Engine/Plugins/Media/PixelCapture/Source/PixelCapture/Public/PixelCaptureCapturerRHIToI420Compute.h b/Engine/Plugins/Media/PixelCapture/Source/PixelCapture/Public/PixelCaptureCapturerRHIToI420Compute.h
--- a/Engine/Plugins/Media/PixelCapture/Source/PixelCapture/Public/PixelCaptureCapturerRHIToI420Compute.h
+++ b/Engine/Plugins/Media/PixelCapture/Source/PixelCapture/Public/PixelCaptureCapturerRHIToI420Compute.h
@@ -24,10 +24,10 @@ public:
 	UE_API virtual ~FPixelCaptureCapturerRHIToI420Compute();
 
 protected:
-	virtual FString					  GetCapturerName() const override { return "RHIToI420Compute"; }
-	UE_API virtual void					  Initialize(int32 InputWidth, int32 InputHeight) override;
-	UE_API virtual IPixelCaptureOutputFrame* CreateOutputBuffer(int32 InputWidth, int32 InputHeight) override;
-	UE_API virtual void					  BeginProcess(const IPixelCaptureInputFrame& InputFrame, TSharedPtr<IPixelCaptureOutputFrame> OutputBuffer) override;
+	virtual FString                          GetCapturerName() const override { return "RHIToI420Compute"; }
+	UE_API virtual void                      Initialize(FRHICommandListBase& RHICmdList, int32 InputWidth, int32 InputHeight) override;
+	UE_API virtual IPixelCaptureOutputFrame* CreateOutputBuffer(FRHICommandListBase& RHICmdList, int32 InputWidth, int32 InputHeight) override;
+	UE_API virtual void                      BeginProcess(FRHICommandList& RHICmdList, const IPixelCaptureInputFrame& InputFrame, TSharedPtr<IPixelCaptureOutputFrame> OutputBuffer) override;
 
 private:
 	// dimensions of the texures
diff --git a/Engine/Plugins/Media/PixelCapture/Source/PixelCaptureShaders/Private/RGBToYUVShader.cpp b/Engine/Plugins/Media/PixelCapture/Source/PixelCaptureShaders/Private/RGBToYUVShader.cpp
--- a/Engine/Plugins/Media/PixelCapture/Source/PixelCaptureShaders/Private/RGBToYUVShader.cpp
+++ b/Engine/Plugins/Media/PixelCapture/Source/PixelCaptureShaders/Private/RGBToYUVShader.cpp
@@ -41,7 +41,7 @@ public:
 
 IMPLEMENT_GLOBAL_SHADER(FExtractI420CS, "/Plugin/PixelCapture/Private/RGBToYUV.usf", "ExtractI420", SF_Compute);
 
-void FRGBToYUVShader::Dispatch(FRHICommandListImmediate& RHICmdList, const FRGBToYUVShaderParameters& InParameters)
+void FRGBToYUVShader::Dispatch(FRHIComputeCommandList& RHICmdList, const FRGBToYUVShaderParameters& InParameters)
 {
 	FExtractI420CS::FParameters PassParametersY;
 	PassParametersY.SourceTexture = InParameters.SourceTexture;
diff --git a/Engine/Plugins/Media/PixelCapture/Source/PixelCaptureShaders/Public/RGBToYUVShader.h b/Engine/Plugins/Media/PixelCapture/Source/PixelCaptureShaders/Public/RGBToYUVShader.h
--- a/Engine/Plugins/Media/PixelCapture/Source/PixelCaptureShaders/Public/RGBToYUVShader.h
+++ b/Engine/Plugins/Media/PixelCapture/Source/PixelCaptureShaders/Public/RGBToYUVShader.h
@@ -24,7 +24,7 @@ struct FRGBToYUVShaderParameters
 class FRGBToYUVShader
 {
 public:
-	static UE_API void Dispatch(FRHICommandListImmediate& RHICmdList, const FRGBToYUVShaderParameters& InParameters);
+	static UE_API void Dispatch(FRHIComputeCommandList& RHICmdList, const FRGBToYUVShaderParameters& InParameters);
 };
 
 #undef UE_API
diff --git a/Engine/Plugins/Media/PixelStreaming/Source/PixelStreaming/Private/PixelStreamingVideoInput.cpp b/Engine/Plugins/Media/PixelStreaming/Source/PixelStreaming/Private/PixelStreamingVideoInput.cpp
--- a/Engine/Plugins/Media/PixelStreaming/Source/PixelStreaming/Private/PixelStreamingVideoInput.cpp
+++ b/Engine/Plugins/Media/PixelStreaming/Source/PixelStreaming/Private/PixelStreamingVideoInput.cpp
@@ -5,6 +5,7 @@
 #include "FrameBufferMultiFormat.h"
 #include "PixelCaptureBufferFormat.h"
 #include "PixelStreamingTrace.h"
+#include "RHICommandList.h"
 
 uint32 FPixelStreamingVideoInput::NextStreamId = 1;
 
@@ -36,7 +37,7 @@ void FPixelStreamingVideoInput::OnFrame(const IPixelCaptureInputFrame& InputFram
 	LastFrameWidth = InputFrame.GetWidth();
 	LastFrameHeight = InputFrame.GetHeight();
 
-	FrameCapturer->Capture(InputFrame);
+	FrameCapturer->Capture(FRHICommandListImmediate::Get(), InputFrame);
 }
 
 rtc::scoped_refptr<webrtc::VideoFrameBuffer> FPixelStreamingVideoInput::GetFrameBuffer()
diff --git a/Engine/Plugins/Media/PixelStreaming2/Source/PixelStreaming2/Private/VideoCapturer.cpp b/Engine/Plugins/Media/PixelStreaming2/Source/PixelStreaming2/Private/VideoCapturer.cpp
--- a/Engine/Plugins/Media/PixelStreaming2/Source/PixelStreaming2/Private/VideoCapturer.cpp
+++ b/Engine/Plugins/Media/PixelStreaming2/Source/PixelStreaming2/Private/VideoCapturer.cpp
@@ -14,6 +14,7 @@
 #include "PixelStreaming2PluginSettings.h"
 #include "Logging.h"
 #include "PixelStreaming2Trace.h"
+#include "RHICommandList.h"
 #include "UtilsVideo.h"
 
 namespace UE::PixelStreaming2
@@ -104,7 +105,7 @@ namespace UE::PixelStreaming2
 		LastFrameType = InputFrame.GetType();
 		bReady = true;
 
-		FrameCapturer->Capture(InputFrame);
+		FrameCapturer->Capture(FRHICommandListImmediate::Get(), InputFrame);
 	}
 
 	TSharedPtr<IPixelCaptureOutputFrame> FVideoCapturer::WaitForFormat(int32 Format, TOptional<FIntPoint> Resolution, uint32 MaxWaitTime)
