Removing implied RHICommandListImmediate uses.

diff --git a/Engine/Plugins/Media/PixelCapture/Source/PixelCapture/Private/PixelCaptureCapturerI420ToRHI.cpp b/Engine/Plugins/Media/PixelCapture/Source/PixelCapture/Private/PixelCaptureCapturerI420ToRHI.cpp
--- a/Engine/Plugins/Media/PixelCapture/Source/PixelCapture/Private/PixelCaptureCapturerI420ToRHI.cpp
+++ b/Engine/Plugins/Media/PixelCapture/Source/PixelCapture/Private/PixelCaptureCapturerI420ToRHI.cpp
@@ -66,7 +66,8 @@ IPixelCaptureOutputFrame* FPixelCaptureCapturerI420ToRHI::CreateOutputBuffer(int
 		TextureDesc.AddFlags(ETextureCreateFlags::Shared);
 	}
 
-	return new FPixelCaptureOutputFrameRHI(RHICreateTexture(TextureDesc));
+	FRHICommandListImmediate& RHICmdList = FRHICommandListImmediate::Get();
+	return new FPixelCaptureOutputFrameRHI(RHICmdList.CreateTexture(TextureDesc));
 }
 
 void FPixelCaptureCapturerI420ToRHI::BeginProcess(const IPixelCaptureInputFrame& InputFrame, TSharedPtr<IPixelCaptureOutputFrame> OutputBuffer)
diff --git a/Engine/Plugins/Media/PixelCapture/Source/PixelCapture/Private/PixelCaptureCapturerNV12ToRHI.cpp b/Engine/Plugins/Media/PixelCapture/Source/PixelCapture/Private/PixelCaptureCapturerNV12ToRHI.cpp
--- a/Engine/Plugins/Media/PixelCapture/Source/PixelCapture/Private/PixelCaptureCapturerNV12ToRHI.cpp
+++ b/Engine/Plugins/Media/PixelCapture/Source/PixelCapture/Private/PixelCaptureCapturerNV12ToRHI.cpp
@@ -65,7 +65,8 @@ IPixelCaptureOutputFrame* FPixelCaptureCapturerNV12ToRHI::CreateOutputBuffer(int
 		TextureDesc.AddFlags(ETextureCreateFlags::Shared);
 	}
 
-	return new FPixelCaptureOutputFrameRHI(RHICreateTexture(TextureDesc), { InputWidth, InputHeight, EPixelFormat::PF_NV12 });
+	FRHICommandListImmediate& RHICmdList = FRHICommandListImmediate::Get();
+	return new FPixelCaptureOutputFrameRHI(RHICmdList.CreateTexture(TextureDesc), { InputWidth, InputHeight, EPixelFormat::PF_NV12 });
 }
 
 void FPixelCaptureCapturerNV12ToRHI::BeginProcess(const IPixelCaptureInputFrame& InputFrame, TSharedPtr<IPixelCaptureOutputFrame> OutputBuffer)
@@ -91,7 +92,7 @@ void FPixelCaptureCapturerNV12ToRHI::BeginProcess(const IPixelCaptureInputFrame&
 
 				// update our staging texture
 				FUpdateTextureRegion2D Region(0, 0, 0, 0, Self->Config.OutputResolution.X, Self->Config.OutputResolution.Y * 1.5);
-				GDynamicRHI->RHIUpdateTexture2D(RHICmdList, OutputTexture->GetFrameTexture(), 0, Region, Self->Config.OutputResolution.X, NV12Buffer->GetData());
+				RHICmdList.UpdateTexture2D(OutputTexture->GetFrameTexture(), 0, Region, Self->Config.OutputResolution.X, NV12Buffer->GetData());
 
 				UE::PixelCapture::MarkGPUWorkEnd(OutputTexture);
 
diff --git a/Engine/Plugins/Media/PixelCapture/Source/PixelCapture/Private/PixelCaptureCapturerRHI.cpp b/Engine/Plugins/Media/PixelCapture/Source/PixelCapture/Private/PixelCaptureCapturerRHI.cpp
--- a/Engine/Plugins/Media/PixelCapture/Source/PixelCapture/Private/PixelCaptureCapturerRHI.cpp
+++ b/Engine/Plugins/Media/PixelCapture/Source/PixelCapture/Private/PixelCaptureCapturerRHI.cpp
@@ -52,7 +52,8 @@ IPixelCaptureOutputFrame* FPixelCaptureCapturerRHI::CreateOutputBuffer(int32 Inp
 		TextureDesc.AddFlags(ETextureCreateFlags::Shared);
 	}
 
-	return new FPixelCaptureOutputFrameRHI(RHICreateTexture(TextureDesc));
+	FRHICommandListImmediate& RHICmdList = FRHICommandListImmediate::Get();
+	return new FPixelCaptureOutputFrameRHI(RHICmdList.CreateTexture(TextureDesc));
 }
 
 void FPixelCaptureCapturerRHI::BeginProcess(const IPixelCaptureInputFrame& InputFrame, TSharedPtr<IPixelCaptureOutputFrame> OutputBuffer)
diff --git a/Engine/Plugins/Media/PixelCapture/Source/PixelCapture/Private/PixelCaptureCapturerRHIRDG.cpp b/Engine/Plugins/Media/PixelCapture/Source/PixelCapture/Private/PixelCaptureCapturerRHIRDG.cpp
--- a/Engine/Plugins/Media/PixelCapture/Source/PixelCapture/Private/PixelCaptureCapturerRHIRDG.cpp
+++ b/Engine/Plugins/Media/PixelCapture/Source/PixelCapture/Private/PixelCaptureCapturerRHIRDG.cpp
@@ -51,7 +51,8 @@ IPixelCaptureOutputFrame* FPixelCaptureCapturerRHIRDG::CreateOutputBuffer(int32
 		TextureDesc.AddFlags(ETextureCreateFlags::Shared);
 	}
 
-	return new FPixelCaptureOutputFrameRHI(RHICreateTexture(TextureDesc));
+	FRHICommandListImmediate& RHICmdList = FRHICommandListImmediate::Get();
+	return new FPixelCaptureOutputFrameRHI(RHICmdList.CreateTexture(TextureDesc));
 }
 
 void FPixelCaptureCapturerRHIRDG::BeginProcess(const IPixelCaptureInputFrame& InputFrame, TSharedPtr<IPixelCaptureOutputFrame> OutputBuffer)
diff --git a/Engine/Plugins/Media/PixelCapture/Source/PixelCapture/Private/PixelCaptureCapturerRHIToI420CPU.cpp b/Engine/Plugins/Media/PixelCapture/Source/PixelCapture/Private/PixelCaptureCapturerRHIToI420CPU.cpp
--- a/Engine/Plugins/Media/PixelCapture/Source/PixelCapture/Private/PixelCaptureCapturerRHIToI420CPU.cpp
+++ b/Engine/Plugins/Media/PixelCapture/Source/PixelCapture/Private/PixelCaptureCapturerRHIToI420CPU.cpp
@@ -50,7 +50,8 @@ void FPixelCaptureCapturerRHIToI420CPU::Initialize(int32 InputWidth, int32 Input
 		TextureDesc.AddFlags(ETextureCreateFlags::Shared);
 	}
 
-	StagingTexture = RHICreateTexture(TextureDesc);
+	FRHICommandListImmediate& RHICmdList = FRHICommandListImmediate::Get();
+	StagingTexture = RHICmdList.CreateTexture(TextureDesc);
 
 	TextureReader = MakeShared<FRHIGPUTextureReadback>(TEXT("FPixelCaptureCapturerRHIToI420CPUReadback"));
 
@@ -135,6 +136,8 @@ void FPixelCaptureCapturerRHIToI420CPU::OnRHIStageComplete(TSharedPtr<IPixelCapt
 
 void FPixelCaptureCapturerRHIToI420CPU::Initialize(int32 InputWidth, int32 InputHeight)
 {
+	FRHICommandListImmediate& RHICmdList = FRHICommandListImmediate::Get();
+
 	FRHITextureCreateDesc StagingTextureDesc =
 		FRHITextureCreateDesc::Create2D(TEXT("FPixelCaptureCapturerRHIToI420CPU StagingTexture"), Config.OutputResolution.X, Config.OutputResolution.Y, EPixelFormat::PF_B8G8R8A8)
 			.SetClearValue(FClearValueBinding::None)
@@ -151,7 +154,7 @@ void FPixelCaptureCapturerRHIToI420CPU::Initialize(int32 InputWidth, int32 Input
 		StagingTextureDesc.AddFlags(ETextureCreateFlags::Shared);
 	}
 
-	StagingTexture = RHICreateTexture(StagingTextureDesc);
+	StagingTexture = RHICmdList.CreateTexture(StagingTextureDesc);
 
 	FRHITextureCreateDesc ReadbackTextureDesc = FRHITextureCreateDesc::Create2D(TEXT("FPixelCaptureCapturerRHIToI420CPU ReadbackTexture"), Config.OutputResolution.X, Config.OutputResolution.Y, EPixelFormat::PF_B8G8R8A8)
 													.SetClearValue(FClearValueBinding::None)
@@ -159,7 +162,7 @@ void FPixelCaptureCapturerRHIToI420CPU::Initialize(int32 InputWidth, int32 Input
 													.SetInitialState(ERHIAccess::CopyDest)
 													.DetermineInititialState();
 
-	ReadbackTexture = RHICreateTexture(ReadbackTextureDesc);
+	ReadbackTexture = RHICmdList.CreateTexture(ReadbackTextureDesc);
 
 	FPixelCaptureCapturer::Initialize(InputWidth, InputHeight);
 }
diff --git a/Engine/Plugins/Media/PixelCapture/Source/PixelCapture/Private/PixelCaptureCapturerRHIToI420Compute.cpp b/Engine/Plugins/Media/PixelCapture/Source/PixelCapture/Private/PixelCaptureCapturerRHIToI420Compute.cpp
--- a/Engine/Plugins/Media/PixelCapture/Source/PixelCapture/Private/PixelCaptureCapturerRHIToI420Compute.cpp
+++ b/Engine/Plugins/Media/PixelCapture/Source/PixelCapture/Private/PixelCaptureCapturerRHIToI420Compute.cpp
@@ -57,9 +57,9 @@ void FPixelCaptureCapturerRHIToI420Compute::Initialize(int32 InputWidth, int32 I
 			.SetInitialState(ERHIAccess::UAVCompute)
 			.DetermineInititialState();
 
-	TextureY = RHICreateTexture(TextureDescY);
-	TextureU = RHICreateTexture(TextureDescUV);
-	TextureV = RHICreateTexture(TextureDescUV);
+	TextureY = RHICmdList.CreateTexture(TextureDescY);
+	TextureU = RHICmdList.CreateTexture(TextureDescUV);
+	TextureV = RHICmdList.CreateTexture(TextureDescUV);
 
 	FRHITextureCreateDesc StagingDescY =
 		FRHITextureCreateDesc::Create2D(TEXT("YUV Output CPU Texture"), PlaneYDimensions.X, PlaneYDimensions.Y, EPixelFormat::PF_R8)
@@ -75,9 +75,9 @@ void FPixelCaptureCapturerRHIToI420Compute::Initialize(int32 InputWidth, int32 I
 			.SetInitialState(ERHIAccess::Unknown)
 			.DetermineInititialState();
 
-	StagingTextureY = RHICreateTexture(StagingDescY);
-	StagingTextureU = RHICreateTexture(StagingDescUV);
-	StagingTextureV = RHICreateTexture(StagingDescUV);
+	StagingTextureY = RHICmdList.CreateTexture(StagingDescY);
+	StagingTextureU = RHICmdList.CreateTexture(StagingDescUV);
+	StagingTextureV = RHICmdList.CreateTexture(StagingDescUV);
 
 	TextureYUAV = RHICmdList.CreateUnorderedAccessView(TextureY, FRHIViewDesc::CreateTextureUAV().SetDimensionFromTexture(TextureY));
 	TextureUUAV = RHICmdList.CreateUnorderedAccessView(TextureU, FRHIViewDesc::CreateTextureUAV().SetDimensionFromTexture(TextureU));
