[PixelCapture] Add: Support for iOS and Android

#Synced-CL 48443255

diff --git a/Engine/Plugins/Media/PixelCapture/PixelCapture.uplugin b/Engine/Plugins/Media/PixelCapture/PixelCapture.uplugin
--- a/Engine/Plugins/Media/PixelCapture/PixelCapture.uplugin
+++ b/Engine/Plugins/Media/PixelCapture/PixelCapture.uplugin
@@ -19,14 +19,14 @@
 			"Name": "PixelCapture",
 			"Type": "Runtime",
 			"LoadingPhase": "PreDefault",
-			"PlatformAllowList": [ "Win64", "Linux", "Mac" ],
+			"PlatformAllowList": [ "Win64", "Linux", "Mac", "Android", "IOS" ],
 			"TargetDenyList": [ "Server" ]
 		},
 		{
             "Name": "PixelCaptureShaders",
             "Type": "Runtime",
             "LoadingPhase": "PostConfigInit",
-            "PlatformAllowList": [ "Win64", "Linux", "Mac" ],
+            "PlatformAllowList": [ "Win64", "Linux", "Mac", "Android", "IOS" ],
             "TargetDenyList": [ "Server" ]
 		}
 	],
diff --git a/Engine/Plugins/Media/PixelCapture/Source/PixelCapture/Internal/PixelCaptureCapturerMediaCapture.h b/Engine/Plugins/Media/PixelCapture/Source/PixelCapture/Internal/PixelCaptureCapturerMediaCapture.h
--- a/Engine/Plugins/Media/PixelCapture/Source/PixelCapture/Internal/PixelCaptureCapturerMediaCapture.h
+++ b/Engine/Plugins/Media/PixelCapture/Source/PixelCapture/Internal/PixelCaptureCapturerMediaCapture.h
@@ -104,14 +104,14 @@ protected:
 
 	virtual bool SupportsAnyThreadCapture() const
 	{
-#if PLATFORM_MAC
-		// On Mac, Media Capture must capture cpu frames on the render thread.
+#if PLATFORM_APPLE
+		// On Apple, Media Capture must capture cpu frames on the render thread.
 		// This is because MediaCapture Lock_Unsafe uses RHIMapStagingSurface
 		// which needs to run on the rendering thread with Metal.
 		return Format == PixelCaptureBufferFormat::FORMAT_RHI;
-#else  // PLATFORM_MAC
+#else  // PLATFORM_APPLE
 		return true;
-#endif // PLATFORM_MAC
+#endif // PLATFORM_APPLE
 	}
 
 	UE_API virtual ETextureCreateFlags GetOutputTextureFlags() const override;
diff --git a/Engine/Plugins/Media/PixelCapture/Source/PixelCapture/Private/PixelCaptureCapturerI420ToRHI.cpp b/Engine/Plugins/Media/PixelCapture/Source/PixelCapture/Private/PixelCaptureCapturerI420ToRHI.cpp
--- a/Engine/Plugins/Media/PixelCapture/Source/PixelCapture/Private/PixelCaptureCapturerI420ToRHI.cpp
+++ b/Engine/Plugins/Media/PixelCapture/Source/PixelCapture/Private/PixelCaptureCapturerI420ToRHI.cpp
@@ -57,6 +57,8 @@ IPixelCaptureOutputFrame* FPixelCaptureCapturerI420ToRHI::CreateOutputBuffer(int
 		TextureDesc.AddFlags(ETextureCreateFlags::SRGB);
 	}
 
+#if !PLATFORM_ANDROID
+	// TODO (william.belcher): Android doesn't seem to like the External/Shared flags - investigate further
 	if (RHIGetInterfaceType() == ERHIInterfaceType::Vulkan)
 	{
 		TextureDesc.AddFlags(ETextureCreateFlags::External);
@@ -65,6 +67,7 @@ IPixelCaptureOutputFrame* FPixelCaptureCapturerI420ToRHI::CreateOutputBuffer(int
 	{
 		TextureDesc.AddFlags(ETextureCreateFlags::Shared);
 	}
+#endif
 
 	FRHICommandListImmediate& RHICmdList = FRHICommandListImmediate::Get();
 	return new FPixelCaptureOutputFrameRHI(RHICmdList.CreateTexture(TextureDesc));
diff --git a/Engine/Plugins/Media/PixelCapture/Source/PixelCapture/Private/PixelCaptureCapturerMediaCapture.cpp b/Engine/Plugins/Media/PixelCapture/Source/PixelCapture/Private/PixelCaptureCapturerMediaCapture.cpp
--- a/Engine/Plugins/Media/PixelCapture/Source/PixelCapture/Private/PixelCaptureCapturerMediaCapture.cpp
+++ b/Engine/Plugins/Media/PixelCapture/Source/PixelCapture/Private/PixelCaptureCapturerMediaCapture.cpp
@@ -86,6 +86,12 @@ void UPixelCaptureMediaCapture::OnFrameCaptured_RenderingThread(
 	OnFrameCaptured_AnyThread(InBaseData, InUserData, ResourceData);
 }
 
+#if PLATFORM_IOS
+BEGIN_SHADER_PARAMETER_STRUCT(FWriteFenceParameters,)
+	SHADER_PARAMETER_RDG_TEXTURE(Texture2D, Texture)
+END_SHADER_PARAMETER_STRUCT()
+#endif // PLATFORM_IOS
+
 void UPixelCaptureMediaCapture::OnCustomCapture_RenderingThread(
 	FRDGBuilder&										   GraphBuilder,
 	const FCaptureBaseData&								   InBaseData,
@@ -99,7 +105,22 @@ void UPixelCaptureMediaCapture::OnCustomCapture_RenderingThread(
 	FGPUFenceRHIRef Fence = GDynamicRHI->RHICreateGPUFence(TEXT("UPixelCaptureMediaCapture Fence"));
 	if (UseExperimentalScheduling() && ShouldCaptureRHIResource())
 	{
+#if PLATFORM_IOS
+		// On iOS, GetRHI is not valid to be called unless the resource has been added.
+		// This allows bAllowRHIAccess to be set so that GetRHI can be called
+		FWriteFenceParameters* PassParams = GraphBuilder.AllocParameters<FWriteFenceParameters>();
+		PassParams->Texture = OutputTexture;
+		GraphBuilder.AddPass(
+			RDG_EVENT_NAME("UPixelCaptureMediaCapture AddFence"),
+			PassParams,
+			ERDGPassFlags::Copy | ERDGPassFlags::NeverCull,								
+			[this, PassParams, Fence](FRDGAsyncTask, FRHICommandList& RHICmdList) {
+				Fences.Add(reinterpret_cast<uintptr_t>(PassParams->Texture->GetRHI()), Fence);
+			}
+		);
+#else // PLATFORM_IOS
 		Fences.Add(reinterpret_cast<uintptr_t>(OutputTexture->GetRHI()), Fence);
+#endif // PLATFORM_IOS
 	}
 
 	bool bRequiresFormatConversion = InSourceTexture->Desc.Format != OutputTexture->Desc.Format;
@@ -117,7 +138,7 @@ void UPixelCaptureMediaCapture::OnCustomCapture_RenderingThread(
 	}
 	else
 	{
-#if PLATFORM_MAC
+#if PLATFORM_APPLE
 		// Create a staging texture that is the same size and format as the final.
 		FRDGTextureRef			   StagingTexture = GraphBuilder.CreateTexture(FRDGTextureDesc::Create2D(FIntPoint(OutputTexture->Desc.Extent.X, OutputTexture->Desc.Extent.Y), OutputTexture->Desc.Format, OutputTexture->Desc.ClearValue, ETextureCreateFlags::RenderTargetable), TEXT("PixelStreamingMediaIOCapture Staging"));
 		FScreenPassTextureViewport StagingViewport(StagingTexture);
@@ -141,7 +162,7 @@ void UPixelCaptureMediaCapture::OnCustomCapture_RenderingThread(
 		FModifyAlphaSwizzleRgbaPS::FParameters*	 PixelShaderParameters = PixelShader->AllocateAndSetParameters(
 			 GraphBuilder,
 			 InSourceTexture,
-#if PLATFORM_MAC
+#if PLATFORM_APPLE
 			StagingTexture
 #else
 			OutputTexture
@@ -155,7 +176,7 @@ void UPixelCaptureMediaCapture::OnCustomCapture_RenderingThread(
 			GraphBuilder,
 			RDG_EVENT_NAME("PixelStreamingEpicRtcMediaIOCapture Swizzle"),
 			FScreenPassViewInfo(),
-#if PLATFORM_MAC
+#if PLATFORM_APPLE
 			StagingViewport,
 #else
 			OutputViewport,
@@ -166,7 +187,7 @@ void UPixelCaptureMediaCapture::OnCustomCapture_RenderingThread(
 			PixelShader,
 			PixelShaderParameters);
 
-#if PLATFORM_MAC
+#if PLATFORM_APPLE
 		// Now we can be certain the formats are the same and size are the same. simple copy
 		AddDrawTexturePass(
 			GraphBuilder,
@@ -189,7 +210,7 @@ void UPixelCaptureMediaCapture::OnCustomCapture_RenderingThread(
 
 ETextureCreateFlags UPixelCaptureMediaCapture::GetOutputTextureFlags() const
 {
-#if PLATFORM_MAC
+#if PLATFORM_APPLE
 	return TexCreate_CPUReadback;
 #else
 	ETextureCreateFlags Flags = TexCreate_RenderTargetable | TexCreate_UAV;
@@ -214,8 +235,19 @@ ETextureCreateFlags UPixelCaptureMediaCapture::GetOutputTextureFlags() const
 
 void UPixelCaptureMediaCapture::WaitForGPU(FRHITexture* InRHITexture)
 {
+	constexpr double SleepTimeSeconds = 50 * 1E-6;
 	FGPUFenceRHIRef Fence;
+	
+#if PLATFORM_IOS
+	const double MaxWaitTime = FPlatformTime::Seconds() + 0.1;
+	while (!Fences.RemoveAndCopyValue(reinterpret_cast<uintptr_t>(InRHITexture), Fence) && FPlatformTime::Seconds() < MaxWaitTime)
+	{
+		FPlatformProcess::SleepNoStats(SleepTimeSeconds);
+	}
+	if (!Fence.IsValid())
+#else // PLATFORM_IOS
 	if (!Fences.RemoveAndCopyValue(reinterpret_cast<uintptr_t>(InRHITexture), Fence))
+#endif // PLATFORM_IOS
 	{
 		UE_LOGFMT(LogPixelCapture, Error, "UPixelCaptureMediaCapture::WaitForGPU: No associating fence!");
 		return;
@@ -223,7 +255,6 @@ void UPixelCaptureMediaCapture::WaitForGPU(FRHITexture* InRHITexture)
 
 	while (!Fence->Poll())
 	{
-		constexpr float SleepTimeSeconds = 50 * 1E-6;
 		FPlatformProcess::SleepNoStats(SleepTimeSeconds);
 	}
 }
diff --git a/Engine/Plugins/Media/PixelCapture/Source/PixelCapture/Private/PixelCaptureUtils.h b/Engine/Plugins/Media/PixelCapture/Source/PixelCapture/Private/PixelCaptureUtils.h
--- a/Engine/Plugins/Media/PixelCapture/Source/PixelCapture/Private/PixelCaptureUtils.h
+++ b/Engine/Plugins/Media/PixelCapture/Source/PixelCapture/Private/PixelCaptureUtils.h
@@ -115,7 +115,7 @@ inline void CopyTextureRDG(FRHICommandListImmediate& RHICmdList, FTextureRHIRef
 		// The todo is here because there has to be a better way to achieve what we want without
 		// all of this song and dance.
 
-#if PLATFORM_MAC
+#if PLATFORM_APPLE
 		// Create a staging texture that is the same size and format as the final.
 		FRDGTextureRef StagingTexture = GraphBuilder.CreateTexture(FRDGTextureDesc::Create2D(FIntPoint(OutputTexture->Desc.Extent.X, OutputTexture->Desc.Extent.Y), OutputTexture->Desc.Format, OutputTexture->Desc.ClearValue, ETextureCreateFlags::RenderTargetable), TEXT("PixelCaptureCopy Staging"));
 		FScreenPassTextureViewport StagingViewport(StagingTexture);
@@ -141,7 +141,7 @@ inline void CopyTextureRDG(FRHICommandListImmediate& RHICmdList, FTextureRHIRef
 		FModifyAlphaSwizzleRgbaPS::FParameters* PixelShaderParameters = PixelShader->AllocateAndSetParameters(
 			GraphBuilder,
 			InputTexture,
-#if PLATFORM_MAC
+#if PLATFORM_APPLE
 			StagingTexture
 #else
 			OutputTexture
@@ -155,7 +155,7 @@ inline void CopyTextureRDG(FRHICommandListImmediate& RHICmdList, FTextureRHIRef
 			GraphBuilder,
 			RDG_EVENT_NAME("PixelCapturerSwizzle"),
 			FScreenPassViewInfo(),
-#if PLATFORM_MAC
+#if PLATFORM_APPLE
 			StagingViewport,
 #else
 			OutputViewport,
@@ -166,7 +166,7 @@ inline void CopyTextureRDG(FRHICommandListImmediate& RHICmdList, FTextureRHIRef
 			PixelShader,
 			PixelShaderParameters);
 
-#if PLATFORM_MAC
+#if PLATFORM_APPLE
 		// Now we can be certain the formats are the same and size are the same. simple copy
 		AddDrawTexturePass(
 			GraphBuilder,
