[PixelCapture] Fix: Remove usage of WebRTC TimeMillis for metadata. Replace with FPlatformTime

#Synced-CL 48443255

diff --git a/Engine/Plugins/Media/PixelCapture/Source/PixelCapture/Public/IPixelCaptureInputFrame.h b/Engine/Plugins/Media/PixelCapture/Source/PixelCapture/Public/IPixelCaptureInputFrame.h
--- a/Engine/Plugins/Media/PixelCapture/Source/PixelCapture/Public/IPixelCaptureInputFrame.h
+++ b/Engine/Plugins/Media/PixelCapture/Source/PixelCapture/Public/IPixelCaptureInputFrame.h
@@ -4,7 +4,6 @@
 
 #include "PixelCaptureFrameMetadata.h"
 #include "HAL/PlatformTime.h"
-#include "rtc_base/time_utils.h"
 
 /**
  * The base interface that is fed into the capture system. This wrapper is fed into the
@@ -14,9 +13,16 @@
 class IPixelCaptureInputFrame
 {
 public:
-	IPixelCaptureInputFrame() { Metadata.SourceTime = rtc::TimeMillis(); }
-
-	IPixelCaptureInputFrame(TSharedPtr<FPixelCaptureUserData> UserData) { Metadata.SourceTime = rtc::TimeMillis(); Metadata.UserData = UserData; }
+	IPixelCaptureInputFrame() 
+	{ 
+		Metadata.SourceTime = FPlatformTime::ToMilliseconds64(FPlatformTime::Cycles64()); 
+	}
+
+	IPixelCaptureInputFrame(TSharedPtr<FPixelCaptureUserData> UserData) 
+	{ 
+		Metadata.SourceTime = FPlatformTime::ToMilliseconds64(FPlatformTime::Cycles64()); 
+		Metadata.UserData = UserData; 
+	}
 
 	virtual ~IPixelCaptureInputFrame() = default;
 
