[PixelCapture] Add: Support for iOS and Android. Remove usage of WebRTC's timeMillis from public API

#Synced-CL 48443255

diff --git a/Engine/Plugins/Media/PixelCapture/PixelCapture.uplugin b/Engine/Plugins/Media/PixelCapture/PixelCapture.uplugin
--- a/Engine/Plugins/Media/PixelCapture/PixelCapture.uplugin
+++ b/Engine/Plugins/Media/PixelCapture/PixelCapture.uplugin
@@ -19,14 +19,14 @@
 			"Name": "PixelCapture",
 			"Type": "Runtime",
 			"LoadingPhase": "PreDefault",
-			"PlatformAllowList": [ "Win64", "Linux", "Mac" ],
+			"PlatformAllowList": [ "Win64", "Linux", "Mac", "Android", "IOS" ],
 			"TargetDenyList": [ "Server" ]
 		},
 		{
             "Name": "PixelCaptureShaders",
             "Type": "Runtime",
             "LoadingPhase": "PostConfigInit",
-            "PlatformAllowList": [ "Win64", "Linux", "Mac" ],
+            "PlatformAllowList": [ "Win64", "Linux", "Mac", "Android", "IOS" ],
             "TargetDenyList": [ "Server" ]
 		}
 	],
diff --git a/Engine/Plugins/Media/PixelCapture/Source/PixelCapture/Internal/PixelCaptureCapturerMediaCapture.h b/Engine/Plugins/Media/PixelCapture/Source/PixelCapture/Internal/PixelCaptureCapturerMediaCapture.h
--- a/Engine/Plugins/Media/PixelCapture/Source/PixelCapture/Internal/PixelCaptureCapturerMediaCapture.h
+++ b/Engine/Plugins/Media/PixelCapture/Source/PixelCapture/Internal/PixelCaptureCapturerMediaCapture.h
@@ -104,14 +104,14 @@ protected:
 
 	virtual bool SupportsAnyThreadCapture() const
 	{
-#if PLATFORM_MAC
+#if PLATFORM_APPLE
 		// On Mac, Media Capture must capture cpu frames on the render thread.
 		// This is because MediaCapture Lock_Unsafe uses RHIMapStagingSurface
 		// which needs to run on the rendering thread with Metal.
 		return Format == PixelCaptureBufferFormat::FORMAT_RHI;
-#else  // PLATFORM_MAC
+#else  // PLATFORM_APPLE
 		return true;
-#endif // PLATFORM_MAC
+#endif // PLATFORM_APPLE
 	}
 
 	UE_API virtual ETextureCreateFlags GetOutputTextureFlags() const override;
diff --git a/Engine/Plugins/Media/PixelCapture/Source/PixelCapture/Private/PixelCaptureCapturerI420ToRHI.cpp b/Engine/Plugins/Media/PixelCapture/Source/PixelCapture/Private/PixelCaptureCapturerI420ToRHI.cpp
--- a/Engine/Plugins/Media/PixelCapture/Source/PixelCapture/Private/PixelCaptureCapturerI420ToRHI.cpp
+++ b/Engine/Plugins/Media/PixelCapture/Source/PixelCapture/Private/PixelCaptureCapturerI420ToRHI.cpp
@@ -57,6 +57,8 @@ IPixelCaptureOutputFrame* FPixelCaptureCapturerI420ToRHI::CreateOutputBuffer(int
 		TextureDesc.AddFlags(ETextureCreateFlags::SRGB);
 	}
 
+#if !PLATFORM_ANDROID
+	// TODO (william.belcher): Android doesn't seem to like the External/Shared flags - investigate further
 	if (RHIGetInterfaceType() == ERHIInterfaceType::Vulkan)
 	{
 		TextureDesc.AddFlags(ETextureCreateFlags::External);
@@ -65,6 +67,7 @@ IPixelCaptureOutputFrame* FPixelCaptureCapturerI420ToRHI::CreateOutputBuffer(int
 	{
 		TextureDesc.AddFlags(ETextureCreateFlags::Shared);
 	}
+#endif
 
 	FRHICommandListImmediate& RHICmdList = FRHICommandListImmediate::Get();
 	return new FPixelCaptureOutputFrameRHI(RHICmdList.CreateTexture(TextureDesc));
diff --git a/Engine/Plugins/Media/PixelCapture/Source/PixelCapture/Private/PixelCaptureCapturerMediaCapture.cpp b/Engine/Plugins/Media/PixelCapture/Source/PixelCapture/Private/PixelCaptureCapturerMediaCapture.cpp
--- a/Engine/Plugins/Media/PixelCapture/Source/PixelCapture/Private/PixelCaptureCapturerMediaCapture.cpp
+++ b/Engine/Plugins/Media/PixelCapture/Source/PixelCapture/Private/PixelCaptureCapturerMediaCapture.cpp
@@ -14,6 +14,10 @@
 
 #include UE_INLINE_GENERATED_CPP_BY_NAME(PixelCaptureCapturerMediaCapture)
 
+BEGIN_SHADER_PARAMETER_STRUCT(FWriteFenceParameters, )
+SHADER_PARAMETER_RDG_TEXTURE(Texture2D, Texture)
+END_SHADER_PARAMETER_STRUCT()
+
 void UPixelCaptureMediaCapture::OnRHIResourceCaptured_AnyThread(
 	const FCaptureBaseData&								   BaseData,
 	TSharedPtr<FMediaCaptureUserData, ESPMode::ThreadSafe> UserData,
@@ -34,10 +38,10 @@ void UPixelCaptureMediaCapture::OnRHIResourceCaptured_AnyThread(
 }
 
 void UPixelCaptureMediaCapture::OnRHIResourceCaptured_RenderingThread(
-	FRHICommandListImmediate& RHICmdList,
-	const FCaptureBaseData& BaseData,
+	FRHICommandListImmediate&							   RHICmdList,
+	const FCaptureBaseData&								   BaseData,
 	TSharedPtr<FMediaCaptureUserData, ESPMode::ThreadSafe> UserData,
-	FTextureRHIRef Texture)
+	FTextureRHIRef										   Texture)
 {
 	OnRHIResourceCaptured_AnyThread(BaseData, UserData, Texture);
 }
@@ -96,12 +100,6 @@ void UPixelCaptureMediaCapture::OnCustomCapture_RenderingThread(
 	FVector2D											   CropU,
 	FVector2D											   CropV)
 {
-	FGPUFenceRHIRef Fence = GDynamicRHI->RHICreateGPUFence(TEXT("UPixelCaptureMediaCapture Fence"));
-	if (UseExperimentalScheduling() && ShouldCaptureRHIResource())
-	{
-		Fences.Add(reinterpret_cast<uintptr_t>(OutputTexture->GetRHI()), Fence);
-	}
-
 	bool bRequiresFormatConversion = InSourceTexture->Desc.Format != OutputTexture->Desc.Format;
 	if (InSourceTexture->Desc.Format == OutputTexture->Desc.Format
 		&& InSourceTexture->Desc.Extent.X == OutputTexture->Desc.Extent.X
@@ -117,7 +115,7 @@ void UPixelCaptureMediaCapture::OnCustomCapture_RenderingThread(
 	}
 	else
 	{
-#if PLATFORM_MAC
+#if PLATFORM_APPLE
 		// Create a staging texture that is the same size and format as the final.
 		FRDGTextureRef			   StagingTexture = GraphBuilder.CreateTexture(FRDGTextureDesc::Create2D(FIntPoint(OutputTexture->Desc.Extent.X, OutputTexture->Desc.Extent.Y), OutputTexture->Desc.Format, OutputTexture->Desc.ClearValue, ETextureCreateFlags::RenderTargetable), TEXT("PixelStreamingMediaIOCapture Staging"));
 		FScreenPassTextureViewport StagingViewport(StagingTexture);
@@ -141,7 +139,7 @@ void UPixelCaptureMediaCapture::OnCustomCapture_RenderingThread(
 		FModifyAlphaSwizzleRgbaPS::FParameters*	 PixelShaderParameters = PixelShader->AllocateAndSetParameters(
 			 GraphBuilder,
 			 InSourceTexture,
-#if PLATFORM_MAC
+#if PLATFORM_APPLE
 			StagingTexture
 #else
 			OutputTexture
@@ -155,7 +153,7 @@ void UPixelCaptureMediaCapture::OnCustomCapture_RenderingThread(
 			GraphBuilder,
 			RDG_EVENT_NAME("PixelStreamingEpicRtcMediaIOCapture Swizzle"),
 			FScreenPassViewInfo(),
-#if PLATFORM_MAC
+#if PLATFORM_APPLE
 			StagingViewport,
 #else
 			OutputViewport,
@@ -166,7 +164,7 @@ void UPixelCaptureMediaCapture::OnCustomCapture_RenderingThread(
 			PixelShader,
 			PixelShaderParameters);
 
-#if PLATFORM_MAC
+#if PLATFORM_APPLE
 		// Now we can be certain the formats are the same and size are the same. simple copy
 		AddDrawTexturePass(
 			GraphBuilder,
@@ -177,19 +175,33 @@ void UPixelCaptureMediaCapture::OnCustomCapture_RenderingThread(
 #endif
 	}
 
-	// clang-format off
-	GraphBuilder.AddPass(
-		RDG_EVENT_NAME("UPixelCaptureMediaCapture WriteFence"), 
-		ERDGPassFlags::NeverCull,								
-		[Fence](FRDGAsyncTask, FRHICommandList& RHICmdList) {
-			RHICmdList.WriteGPUFence(Fence);
+	if (UseExperimentalScheduling() && ShouldCaptureRHIResource())
+	{
+		TWeakObjectPtr<UPixelCaptureMediaCapture> WeakThis(this);
+
+		FWriteFenceParameters* PassParams = GraphBuilder.AllocParameters<FWriteFenceParameters>();
+		PassParams->Texture = OutputTexture;
+
+		// clang-format off
+		GraphBuilder.AddPass(
+			RDG_EVENT_NAME("UPixelCaptureMediaCapture WriteFence"), 
+			PassParams,
+			ERDGPassFlags::Copy | ERDGPassFlags::NeverCull,								
+			[WeakThis, PassParams](FRDGAsyncTask, FRHICommandList& RHICmdList) {
+				if (TStrongObjectPtr<UPixelCaptureMediaCapture> This = WeakThis.Pin())
+				{
+					FGPUFenceRHIRef Fence = GDynamicRHI->RHICreateGPUFence(TEXT("UPixelCaptureMediaCapture Fence"));
+					This->Fences.Add(reinterpret_cast<uintptr_t>(PassParams->Texture->GetRHI()), Fence);
+					RHICmdList.WriteGPUFence(Fence);
+				}
 		});
-	// clang-format on
+		// clang-format on
+	}
 }
 
 ETextureCreateFlags UPixelCaptureMediaCapture::GetOutputTextureFlags() const
 {
-#if PLATFORM_MAC
+#if PLATFORM_APPLE
 	return TexCreate_CPUReadback;
 #else
 	ETextureCreateFlags Flags = TexCreate_RenderTargetable | TexCreate_UAV;
@@ -214,16 +226,36 @@ ETextureCreateFlags UPixelCaptureMediaCapture::GetOutputTextureFlags() const
 
 void UPixelCaptureMediaCapture::WaitForGPU(FRHITexture* InRHITexture)
 {
-	FGPUFenceRHIRef Fence;
-	if (!Fences.RemoveAndCopyValue(reinterpret_cast<uintptr_t>(InRHITexture), Fence))
+	constexpr double SleepTimeSeconds = 50 * 1E-6;
+	FGPUFenceRHIRef	 Fence;
+
+	double MaxWaitTime = FPlatformTime::Seconds() + 0.1;
+	while (!Fences.RemoveAndCopyValue(reinterpret_cast<uintptr_t>(InRHITexture), Fence))
 	{
-		UE_LOGFMT(LogPixelCapture, Error, "UPixelCaptureMediaCapture::WaitForGPU: No associating fence!");
+		if (FPlatformTime::Seconds() >= MaxWaitTime)
+		{
+			UE_LOGFMT(LogPixelCapture, Error, "UPixelCaptureMediaCapture::WaitForGPU: Timed out waiting for fence to be added in pass!");
+			return;
+		}
+
+		FPlatformProcess::SleepNoStats(SleepTimeSeconds);
+	}
+
+	if (!Fence.IsValid())
+	{
+		UE_LOGFMT(LogPixelCapture, Error, "UPixelCaptureMediaCapture::WaitForGPU: Found fence is invalid!");
 		return;
 	}
 
+	MaxWaitTime = FPlatformTime::Seconds() + 0.1;
 	while (!Fence->Poll())
 	{
-		constexpr float SleepTimeSeconds = 50 * 1E-6;
+		if (FPlatformTime::Seconds() >= MaxWaitTime)
+		{
+			UE_LOGFMT(LogPixelCapture, Error, "UPixelCaptureMediaCapture::WaitForGPU: Timed out waiting for fence to be written!");
+			return;
+		}
+
 		FPlatformProcess::SleepNoStats(SleepTimeSeconds);
 	}
 }
@@ -299,7 +331,7 @@ void FPixelCaptureCapturerMediaCapture::InitializeMediaCapture()
 	bMediaCaptureInitialized = true;
 }
 
-FString FPixelCaptureCapturerMediaCapture::GetCapturerName() const 
+FString FPixelCaptureCapturerMediaCapture::GetCapturerName() const
 {
 	FString FormatString = TEXT("Unknown");
 	switch (Format)
@@ -311,7 +343,7 @@ FString FPixelCaptureCapturerMediaCapture::GetCapturerName() const
 			FormatString = TEXT("I420");
 			break;
 	}
-	
+
 	return FString::Printf(TEXT("MediaCapture (RHI->%s)"), *FormatString);
 }
 
diff --git a/Engine/Plugins/Media/PixelCapture/Source/PixelCapture/Private/PixelCaptureInputFrameI420.cpp b/Engine/Plugins/Media/PixelCapture/Source/PixelCapture/Private/PixelCaptureInputFrameI420.cpp
--- a/Engine/Plugins/Media/PixelCapture/Source/PixelCapture/Private/PixelCaptureInputFrameI420.cpp
+++ b/Engine/Plugins/Media/PixelCapture/Source/PixelCapture/Private/PixelCaptureInputFrameI420.cpp
@@ -7,7 +7,6 @@ FPixelCaptureInputFrameI420::FPixelCaptureInputFrameI420(TSharedPtr<FPixelCaptur
 	: IPixelCaptureInputFrame(UserData)
 	, I420Buffer(Buffer)
 {
-	Metadata.SourceTime = rtc::TimeMillis();
 }
 
 int32 FPixelCaptureInputFrameI420::GetType() const
diff --git a/Engine/Plugins/Media/PixelCapture/Source/PixelCapture/Private/PixelCaptureInputFrameNV12.cpp b/Engine/Plugins/Media/PixelCapture/Source/PixelCapture/Private/PixelCaptureInputFrameNV12.cpp
--- a/Engine/Plugins/Media/PixelCapture/Source/PixelCapture/Private/PixelCaptureInputFrameNV12.cpp
+++ b/Engine/Plugins/Media/PixelCapture/Source/PixelCapture/Private/PixelCaptureInputFrameNV12.cpp
@@ -7,7 +7,6 @@ FPixelCaptureInputFrameNV12::FPixelCaptureInputFrameNV12(TSharedPtr<FPixelCaptur
 	: IPixelCaptureInputFrame(UserData)
 	, NV12Buffer(Buffer)
 {
-	Metadata.SourceTime = rtc::TimeMillis();
 }
 
 int32 FPixelCaptureInputFrameNV12::GetType() const
diff --git a/Engine/Plugins/Media/PixelCapture/Source/PixelCapture/Private/PixelCaptureInputFrameRHI.cpp b/Engine/Plugins/Media/PixelCapture/Source/PixelCapture/Private/PixelCaptureInputFrameRHI.cpp
--- a/Engine/Plugins/Media/PixelCapture/Source/PixelCapture/Private/PixelCaptureInputFrameRHI.cpp
+++ b/Engine/Plugins/Media/PixelCapture/Source/PixelCapture/Private/PixelCaptureInputFrameRHI.cpp
@@ -7,7 +7,6 @@ FPixelCaptureInputFrameRHI::FPixelCaptureInputFrameRHI(FTextureRHIRef InFrameTex
 	: IPixelCaptureInputFrame(UserData)
 	, FrameTexture(InFrameTexture)
 {
-	Metadata.SourceTime = rtc::TimeMillis();
 }
 
 int32 FPixelCaptureInputFrameRHI::GetType() const
diff --git a/Engine/Plugins/Media/PixelCapture/Source/PixelCapture/Private/PixelCaptureUtils.h b/Engine/Plugins/Media/PixelCapture/Source/PixelCapture/Private/PixelCaptureUtils.h
--- a/Engine/Plugins/Media/PixelCapture/Source/PixelCapture/Private/PixelCaptureUtils.h
+++ b/Engine/Plugins/Media/PixelCapture/Source/PixelCapture/Private/PixelCaptureUtils.h
@@ -115,7 +115,7 @@ inline void CopyTextureRDG(FRHICommandListImmediate& RHICmdList, FTextureRHIRef
 		// The todo is here because there has to be a better way to achieve what we want without
 		// all of this song and dance.
 
-#if PLATFORM_MAC
+#if PLATFORM_APPLE
 		// Create a staging texture that is the same size and format as the final.
 		FRDGTextureRef StagingTexture = GraphBuilder.CreateTexture(FRDGTextureDesc::Create2D(FIntPoint(OutputTexture->Desc.Extent.X, OutputTexture->Desc.Extent.Y), OutputTexture->Desc.Format, OutputTexture->Desc.ClearValue, ETextureCreateFlags::RenderTargetable), TEXT("PixelCaptureCopy Staging"));
 		FScreenPassTextureViewport StagingViewport(StagingTexture);
@@ -141,7 +141,7 @@ inline void CopyTextureRDG(FRHICommandListImmediate& RHICmdList, FTextureRHIRef
 		FModifyAlphaSwizzleRgbaPS::FParameters* PixelShaderParameters = PixelShader->AllocateAndSetParameters(
 			GraphBuilder,
 			InputTexture,
-#if PLATFORM_MAC
+#if PLATFORM_APPLE
 			StagingTexture
 #else
 			OutputTexture
@@ -155,7 +155,7 @@ inline void CopyTextureRDG(FRHICommandListImmediate& RHICmdList, FTextureRHIRef
 			GraphBuilder,
 			RDG_EVENT_NAME("PixelCapturerSwizzle"),
 			FScreenPassViewInfo(),
-#if PLATFORM_MAC
+#if PLATFORM_APPLE
 			StagingViewport,
 #else
 			OutputViewport,
@@ -166,7 +166,7 @@ inline void CopyTextureRDG(FRHICommandListImmediate& RHICmdList, FTextureRHIRef
 			PixelShader,
 			PixelShaderParameters);
 
-#if PLATFORM_MAC
+#if PLATFORM_APPLE
 		// Now we can be certain the formats are the same and size are the same. simple copy
 		AddDrawTexturePass(
 			GraphBuilder,
diff --git a/Engine/Plugins/Media/PixelCapture/Source/PixelCapture/Public/IPixelCaptureInputFrame.h b/Engine/Plugins/Media/PixelCapture/Source/PixelCapture/Public/IPixelCaptureInputFrame.h
--- a/Engine/Plugins/Media/PixelCapture/Source/PixelCapture/Public/IPixelCaptureInputFrame.h
+++ b/Engine/Plugins/Media/PixelCapture/Source/PixelCapture/Public/IPixelCaptureInputFrame.h
@@ -4,7 +4,6 @@
 
 #include "PixelCaptureFrameMetadata.h"
 #include "HAL/PlatformTime.h"
-#include "rtc_base/time_utils.h"
 
 /**
  * The base interface that is fed into the capture system. This wrapper is fed into the
@@ -14,9 +13,16 @@
 class IPixelCaptureInputFrame
 {
 public:
-	IPixelCaptureInputFrame() { Metadata.SourceTime = rtc::TimeMillis(); }
-
-	IPixelCaptureInputFrame(TSharedPtr<FPixelCaptureUserData> UserData) { Metadata.SourceTime = rtc::TimeMillis(); Metadata.UserData = UserData; }
+	IPixelCaptureInputFrame() 
+	{ 
+		Metadata.SourceTime = FPlatformTime::ToMilliseconds64(FPlatformTime::Cycles64()); 
+	}
+
+	IPixelCaptureInputFrame(TSharedPtr<FPixelCaptureUserData> UserData) 
+	{ 
+		Metadata.SourceTime = FPlatformTime::ToMilliseconds64(FPlatformTime::Cycles64()); 
+		Metadata.UserData = UserData; 
+	}
 
 	virtual ~IPixelCaptureInputFrame() = default;
 
