[PixelCapture] Fix: Potential deadlock when pinning weak UObject ptr

diff --git a/Engine/Plugins/Media/PixelCapture/Source/PixelCapture/Private/PixelCaptureCapturerMediaCapture.cpp b/Engine/Plugins/Media/PixelCapture/Source/PixelCapture/Private/PixelCaptureCapturerMediaCapture.cpp
--- a/Engine/Plugins/Media/PixelCapture/Source/PixelCapture/Private/PixelCaptureCapturerMediaCapture.cpp
+++ b/Engine/Plugins/Media/PixelCapture/Source/PixelCapture/Private/PixelCaptureCapturerMediaCapture.cpp
@@ -188,6 +188,11 @@ void UPixelCaptureMediaCapture::OnCustomCapture_RenderingThread(
 			PassParams,
 			ERDGPassFlags::Copy | ERDGPassFlags::NeverCull,								
 			[WeakThis, PassParams](FRDGAsyncTask, FRHICommandList& RHICmdList) {
+				if (!WeakThis.IsValid())
+				{
+					return;
+				}
+				
 				if (TStrongObjectPtr<UPixelCaptureMediaCapture> This = WeakThis.Pin())
 				{
 					FGPUFenceRHIRef Fence = GDynamicRHI->RHICreateGPUFence(TEXT("UPixelCaptureMediaCapture Fence"));
