[PixelCapture] Fix: Missing sRGB flag when required

#Synced-CL 48443255

diff --git a/Engine/Plugins/Media/PixelCapture/Source/PixelCapture/Private/PixelCaptureCapturerMediaCapture.cpp b/Engine/Plugins/Media/PixelCapture/Source/PixelCapture/Private/PixelCaptureCapturerMediaCapture.cpp
--- a/Engine/Plugins/Media/PixelCapture/Source/PixelCapture/Private/PixelCaptureCapturerMediaCapture.cpp
+++ b/Engine/Plugins/Media/PixelCapture/Source/PixelCapture/Private/PixelCaptureCapturerMediaCapture.cpp
@@ -206,10 +206,11 @@ void UPixelCaptureMediaCapture::OnCustomCapture_RenderingThread(
 
 ETextureCreateFlags UPixelCaptureMediaCapture::GetOutputTextureFlags() const
 {
+	ETextureCreateFlags Flags = TexCreate_None;
 #if PLATFORM_APPLE
-	return TexCreate_CPUReadback;
+	Flags |= TexCreate_CPUReadback;
 #else
-	ETextureCreateFlags Flags = TexCreate_RenderTargetable | TexCreate_UAV;
+	Flags |= TexCreate_RenderTargetable | TexCreate_UAV;
 
 	if (RHIGetInterfaceType() == ERHIInterfaceType::Vulkan)
 	{
@@ -219,6 +220,7 @@ ETextureCreateFlags UPixelCaptureMediaCapture::GetOutputTextureFlags() const
 	{
 		Flags |= TexCreate_Shared;
 	}
+#endif
 
 	if (bIsSRGB)
 	{
@@ -226,7 +228,6 @@ ETextureCreateFlags UPixelCaptureMediaCapture::GetOutputTextureFlags() const
 	}
 
 	return Flags;
-#endif
 }
 
 void UPixelCaptureMediaCapture::WaitForGPU(FRHITexture* InRHITexture)
